[SequenceNavigator] Move Motion Design Sequence Navigator to AvalancheSequenceNavigator plugin

diff --git a/Engine/Plugins/Experimental/AvalancheSequenceNavigator/AvalancheSequenceNavigator.uplugin b/Engine/Plugins/Experimental/AvalancheSequenceNavigator/AvalancheSequenceNavigator.uplugin
new file mode 100644
--- /dev/null
+++ b/Engine/Plugins/Experimental/AvalancheSequenceNavigator/AvalancheSequenceNavigator.uplugin
@@ -0,0 +1,41 @@
+﻿{
+	"FileVersion": 3,
+	"Version": 1,
+	"VersionName": "1.0",
+	"FriendlyName": "Motion Design Sequence Navigator Bridge",
+	"Description": "Sequence Navigator Bridge for embedded Motion Design Sequences",
+	"Category": "Animation",
+	"CreatedBy": "Epic Games, Inc.",
+	"CreatedByURL": "https://epicgames.com",
+	"DocsURL": "",
+	"MarketplaceURL": "",
+	"CanContainContent": false,
+	"IsBetaVersion": false,
+	"IsExperimentalVersion": true,
+	"Installed": false,
+	"EnabledByDefault": false,
+	"SupportedTargetPlatforms": [
+		"Win64",
+		"Linux",
+		"Mac"
+	],
+	"Modules": [
+		{
+			"Name": "AvalancheSequenceNavigator",
+			"Type": "Editor",
+			"LoadingPhase": "PostDefault"
+		}
+	],
+	"Plugins": [
+		{
+			"Name": "Avalanche",
+			"Enabled": true,
+			"Optional": true
+		},
+		{
+			"Name": "SequenceNavigator",
+			"Enabled": true,
+			"Optional": true
+		}
+	]
+}
diff --git a/Engine/Plugins/Experimental/AvalancheSequenceNavigator/Source/AvalancheSequenceNavigator/AvalancheSequenceNavigator.Build.cs b/Engine/Plugins/Experimental/AvalancheSequenceNavigator/Source/AvalancheSequenceNavigator/AvalancheSequenceNavigator.Build.cs
new file mode 100644
--- /dev/null
+++ b/Engine/Plugins/Experimental/AvalancheSequenceNavigator/Source/AvalancheSequenceNavigator/AvalancheSequenceNavigator.Build.cs
@@ -0,0 +1,31 @@
+﻿// Copyright Epic Games, Inc. All Rights Reserved.
+
+using UnrealBuildTool;
+
+public class AvalancheSequenceNavigator : ModuleRules
+{
+	public AvalancheSequenceNavigator(ReadOnlyTargetRules Target) : base(Target)
+	{
+		PCHUsage = PCHUsageMode.UseExplicitOrSharedPCHs;
+
+		PublicDependencyModuleNames.AddRange(new string[] {
+			"Core"
+		});
+
+		PrivateDependencyModuleNames.AddRange(new string[] {
+			"Avalanche",
+			"AvalancheEditorCore",
+			"AvalancheSequence",
+			"AvalancheSequencer",
+			"CoreUObject",
+			"Engine",
+			"MovieScene",
+			"Sequencer",
+			"SequenceNavigator",
+			"Slate",
+			"SlateCore",
+			"ToolMenus",
+			"UnrealEd"
+		});
+	}
+}
diff --git a/Engine/Plugins/Experimental/AvalancheSequenceNavigator/Source/AvalancheSequenceNavigator/Private/AvaNavigationToolProvider.cpp b/Engine/Plugins/Experimental/AvalancheSequenceNavigator/Source/AvalancheSequenceNavigator/Private/AvaNavigationToolProvider.cpp
new file mode 100644
--- /dev/null
+++ b/Engine/Plugins/Experimental/AvalancheSequenceNavigator/Source/AvalancheSequenceNavigator/Private/AvaNavigationToolProvider.cpp
@@ -0,0 +1,978 @@
+// Copyright Epic Games, Inc. All Rights Reserved.
+
+#include "AvaNavigationToolProvider.h"
+#include "ActorFactories/ActorFactory.h"
+#include "AvaSequence.h"
+#include "AvaSequenceActor.h"
+#include "AvaSequenceNavigatorCommands.h"
+#include "AvaSequencePlaybackObject.h"
+#include "AvaSequencerUtils.h"
+#include "Columns/AvaNavigationToolStatusColumn.h"
+#include "Columns/NavigationToolColorColumn.h"
+#include "Columns/NavigationToolColumnExtender.h"
+#include "Columns/NavigationToolCommentColumn.h"
+#include "Columns/NavigationToolDeactiveStateColumn.h"
+#include "Columns/NavigationToolHBiasColumn.h"
+#include "Columns/NavigationToolInTimeColumn.h"
+#include "Columns/NavigationToolItemsColumn.h"
+#include "Columns/NavigationToolLabelColumn.h"
+#include "Columns/NavigationToolLengthColumn.h"
+#include "Columns/NavigationToolLockColumn.h"
+#include "Columns/NavigationToolMarkerVisibilityColumn.h"
+#include "Columns/NavigationToolOutTimeColumn.h"
+#include "Columns/NavigationToolPlayheadColumn.h"
+#include "Columns/NavigationToolRevisionControlColumn.h"
+#include "Columns/NavigationToolStartFrameOffsetColumn.h"
+#include "Columns/NavigationToolTakeColumn.h"
+#include "ContentBrowserModule.h"
+#include "DragDrop/NavigationToolAvaSequenceDropHandler.h"
+#include "Editor.h"
+#include "Engine/Level.h"
+#include "Filters/Filters/NavigationToolBuiltInFilterParams.h"
+#include "Framework/Commands/GenericCommands.h"
+#include "Framework/Notifications/NotificationManager.h"
+#include "IAvaSceneInterface.h"
+#include "IAvaSequenceProvider.h"
+#include "IAvaSequencer.h"
+#include "IAvaSequencerProvider.h"
+#include "IContentBrowserSingleton.h"
+#include "INavigationTool.h"
+#include "ISequencer.h"
+#include "ISettingsModule.h"
+#include "ISourceControlModule.h"
+#include "ItemActions/NavigationToolAddItem.h"
+#include "Items/NavigationToolAvaSequence.h"
+#include "Items/NavigationToolItemParameters.h"
+#include "NavigationToolSettings.h"
+#include "ScopedTransaction.h"
+#include "Settings/AvaSequencerSettings.h"
+#include "Widgets/Notifications/SNotificationList.h"
+
+#define LOCTEXT_NAMESPACE "AvaNavigationToolProvider"
+
+namespace UE::AvaSequencer
+{
+
+const FName FAvaNavigationToolProvider::Identifier = TEXT("MotionDesign");
+
+const FText FAvaNavigationToolProvider::MotionDesignColumnViewName = LOCTEXT("MotionDesignColumnViewName", "Motion Design");
+
+const FName FAvaNavigationToolProvider::ToolbarSectionName = TEXT("MotionDesign");
+const FName FAvaNavigationToolProvider::ContextMenuSectionName = TEXT("MotionDesignActions");
+
+FAvaNavigationToolProvider::FAvaNavigationToolProvider(const TSharedRef<IAvaSequencer>& InAvaSequencer)
+	: WeakAvaSequencer(InAvaSequencer)
+	, ToolCommands(MakeShared<FUICommandList>())
+{
+}
+
+FName FAvaNavigationToolProvider::GetIdentifier() const
+{
+	return Identifier;
+}
+
+TSet<TSubclassOf<UMovieSceneSequence>> FAvaNavigationToolProvider::GetSupportedSequenceClasses() const
+{
+	return { UAvaSequence::StaticClass() };
+}
+
+FText FAvaNavigationToolProvider::GetDefaultColumnView() const
+{
+	return MotionDesignColumnViewName;
+}
+
+FNavigationToolSaveState* FAvaNavigationToolProvider::GetSaveState(const UE::SequenceNavigator::INavigationTool& InTool) const
+{
+#if WITH_EDITOR
+	if (IAvaSceneInterface* const AvaScene = GetSceneInterface(InTool))
+	{
+		return &AvaScene->GetNavigationToolSaveState();
+	}
+#endif
+	return nullptr;
+}
+
+void FAvaNavigationToolProvider::SetSaveState(const UE::SequenceNavigator::INavigationTool& InTool, const FNavigationToolSaveState& InSaveState) const
+{
+#if WITH_EDITOR
+	if (IAvaSceneInterface* const AvaScene = GetSceneInterface(InTool))
+	{
+		AvaScene->GetNavigationToolSaveState() = InSaveState;
+	}
+#endif
+}
+
+void FAvaNavigationToolProvider::BindCommands(const TSharedRef<FUICommandList>& InCommandList)
+{
+	const TSharedPtr<IAvaSequencer> AvaSequencer = WeakAvaSequencer.Pin();
+	if (!AvaSequencer.IsValid())
+	{
+		return;
+	}
+
+	InCommandList->Append(ToolCommands);
+
+	const FAvaSequenceNavigatorCommands& AvaSequenceNavigatorCommands = FAvaSequenceNavigatorCommands::Get();
+	const FGenericCommands& GenericCommands = FGenericCommands::Get();
+
+	ToolCommands->MapAction(GenericCommands.Duplicate
+		, FExecuteAction::CreateSP(this, &FAvaNavigationToolProvider::DuplicateSelection)
+		, FCanExecuteAction::CreateSP(this, &FAvaNavigationToolProvider::CanDuplicateSelection));
+
+	ToolCommands->MapAction(GenericCommands.Delete
+		, FExecuteAction::CreateSP(this, &FAvaNavigationToolProvider::DeleteSelection)
+		, FCanExecuteAction::CreateSP(this, &FAvaNavigationToolProvider::CanDeleteSelection));
+
+	ToolCommands->MapAction(GenericCommands.Rename
+		, FExecuteAction::CreateSP(this, &FAvaNavigationToolProvider::RelabelSelection)
+		, FCanExecuteAction::CreateSP(this, &FAvaNavigationToolProvider::CanRelabelSelection));
+
+	ToolCommands->MapAction(AvaSequenceNavigatorCommands.AddNew
+		, FExecuteAction::CreateSP(this, &FAvaNavigationToolProvider::AddSequenceToSelection)
+		, FCanExecuteAction::CreateSP(this, &FAvaNavigationToolProvider::CanAddSequenceToSelection));
+
+	ToolCommands->MapAction(AvaSequenceNavigatorCommands.PlaySelected
+		, FExecuteAction::CreateSP(this, &FAvaNavigationToolProvider::PlaySelection)
+		, FCanExecuteAction::CreateSP(this, &FAvaNavigationToolProvider::CanPlaySelection));
+
+	ToolCommands->MapAction(AvaSequenceNavigatorCommands.ContinueSelected
+		, FExecuteAction::CreateSP(this, &FAvaNavigationToolProvider::ContinueSelection)
+		, FCanExecuteAction::CreateSP(this, &FAvaNavigationToolProvider::CanContinueSelection));
+
+	ToolCommands->MapAction(AvaSequenceNavigatorCommands.StopSelected
+		, FExecuteAction::CreateSP(this, &FAvaNavigationToolProvider::StopSelection)
+		, FCanExecuteAction::CreateSP(this, &FAvaNavigationToolProvider::CanStopSelection));
+
+	ToolCommands->MapAction(AvaSequenceNavigatorCommands.ExportSequence
+		, FExecuteAction::CreateSP(this, &FAvaNavigationToolProvider::ExportSelection)
+		, FCanExecuteAction::CreateSP(this, &FAvaNavigationToolProvider::CanExportSelection));
+
+	ToolCommands->MapAction(AvaSequenceNavigatorCommands.SpawnSequencePlayer
+		, FExecuteAction::CreateSP(this, &FAvaNavigationToolProvider::SpawnPlayersForSelection)
+		, FCanExecuteAction::CreateSP(this, &FAvaNavigationToolProvider::CanSpawnPlayersForSelection));
+}
+
+void FAvaNavigationToolProvider::OnActivate()
+{
+	ExtendToolToolBar();
+	ExtendToolItemContextMenu();
+
+	DragDropInitializedDelegate = UE::SequenceNavigator::FNavigationToolExtender::OnItemDragDropOpInitialized()
+		.AddSP(this, &FAvaNavigationToolProvider::ExtendToolItemDragDropOp);
+}
+
+void FAvaNavigationToolProvider::OnDeactivate()
+{
+	RemoveToolToolBarExtension();
+	RemoveToolItemContextMenuExtension();
+
+	UE::SequenceNavigator::FNavigationToolExtender::OnItemDragDropOpInitialized()
+		.Remove(DragDropInitializedDelegate);
+}
+
+void FAvaNavigationToolProvider::OnExtendColumns(UE::SequenceNavigator::FNavigationToolColumnExtender& OutExtender)
+{
+	using namespace UE::SequenceNavigator;
+
+	// Support built in columns
+	OutExtender.AddColumn<FNavigationToolPlayheadColumn>();
+	OutExtender.AddColumn<FNavigationToolDeactiveStateColumn>();
+	OutExtender.AddColumn<FNavigationToolMarkerVisibilityColumn>();
+	OutExtender.AddColumn<FNavigationToolLockColumn>();
+	OutExtender.AddColumn<FNavigationToolColorColumn>();
+	OutExtender.AddColumn<FNavigationToolLabelColumn>();
+	OutExtender.AddColumn<FNavigationToolItemsColumn>();
+	OutExtender.AddColumn<FNavigationToolInTimeColumn>();
+	OutExtender.AddColumn<FNavigationToolOutTimeColumn>();
+	OutExtender.AddColumn<FNavigationToolLengthColumn>();
+	OutExtender.AddColumn<FNavigationToolHBiasColumn>();
+	OutExtender.AddColumn<FNavigationToolStartFrameOffsetColumn>();
+	OutExtender.AddColumn<FNavigationToolTakeColumn>();
+	OutExtender.AddColumn<FNavigationToolCommentColumn>();
+
+	ISourceControlModule& SourceControlModule = ISourceControlModule::Get();
+	if (SourceControlModule.IsEnabled())
+	{
+		OutExtender.AddColumn<FNavigationToolRevisionControlColumn>();
+	}
+
+	// Add Motion Design specific columns
+	OutExtender.AddColumn<FAvaNavigationToolStatusColumn>();
+
+	FNavigationToolProvider::OnExtendColumns(OutExtender);
+}
+
+void FAvaNavigationToolProvider::OnExtendColumnViews(TSet<FNavigationToolColumnView>& OutColumnViews)
+{
+	using namespace UE::SequenceNavigator;
+
+	FNavigationToolColumnView& MotionDesignColumnView = OutColumnViews
+		.FindOrAdd(FNavigationToolColumnView(MotionDesignColumnViewName));
+
+	MotionDesignColumnView.VisibleColumns =
+		{
+			FNavigationToolColorColumn::StaticColumnId(),
+			FNavigationToolLabelColumn::StaticColumnId(),
+			FNavigationToolItemsColumn::StaticColumnId(),
+			FAvaNavigationToolStatusColumn::StaticColumnId()
+		};
+
+	FNavigationToolProvider::OnExtendColumnViews(OutColumnViews);
+}
+
+void FAvaNavigationToolProvider::OnExtendItemChildren(UE::SequenceNavigator::INavigationTool& InTool
+	, const UE::SequenceNavigator::FNavigationToolItemPtr& InParentItem
+	, TArray<UE::SequenceNavigator::FNavigationToolItemPtr>& OutChildren
+	, const bool bInRecursive)
+{
+	using namespace UE::SequenceNavigator;
+
+	FNavigationToolProvider::OnExtendItemChildren(InTool, InParentItem, OutChildren, bInRecursive);
+
+	// Only extending root item
+	if (InParentItem->GetItemId() != FNavigationToolItemId::RootId)
+	{
+		return;
+	}
+
+	IAvaSequenceProvider* const SequenceProvider = GetSequenceProvider(InTool);
+	if (!SequenceProvider)
+	{
+		return;
+	}
+
+	TArray<TWeakObjectPtr<UAvaSequence>> RootWeakSequences = SequenceProvider->GetRootSequences();
+
+	// Remove invalid sequences and sort
+	RootWeakSequences.RemoveAll([](const TWeakObjectPtr<UAvaSequence>& InRoot)
+		{
+			return !InRoot.IsValid();
+		});
+	RootWeakSequences.Sort([](const TWeakObjectPtr<UAvaSequence>& InA, const TWeakObjectPtr<UAvaSequence>& InB)
+		{
+			return InA->GetDisplayName().CompareTo(InB->GetDisplayName()) > 0;
+		});
+
+	const TSharedRef<FAvaNavigationToolProvider> Provider = SharedThis(this);
+
+	// Add child sequence items
+	for (const TWeakObjectPtr<UAvaSequence>& RootSequenceWeak : RootWeakSequences)
+	{
+		if (UAvaSequence* const RootSequence = RootSequenceWeak.Get())
+		{
+			const FNavigationToolItemPtr NewItem = InTool.FindOrAdd<FNavigationToolAvaSequence>(Provider
+				, InParentItem, RootSequence);
+			const FNavigationToolItemFlagGuard Guard(NewItem, ENavigationToolItemFlags::IgnorePendingKill);
+			OutChildren.Add(NewItem);
+			if (bInRecursive)
+			{
+				NewItem->FindChildren(OutChildren, bInRecursive);
+			}
+		}
+	}
+}
+
+void FAvaNavigationToolProvider::OnExtendBuiltInFilters(TArray<FNavigationToolBuiltInFilterParams>& OutFilterParams)
+{
+	OutFilterParams.Add(FNavigationToolBuiltInFilterParams::CreateSequenceFilter());
+	OutFilterParams.Add(FNavigationToolBuiltInFilterParams::CreateTrackFilter());
+	OutFilterParams.Add(FNavigationToolBuiltInFilterParams::CreateBindingFilter());
+	OutFilterParams.Add(FNavigationToolBuiltInFilterParams::CreateMarkerFilter());
+
+	FNavigationToolProvider::OnExtendBuiltInFilters(OutFilterParams);
+}
+
+TSharedPtr<SequenceNavigator::INavigationTool> FAvaNavigationToolProvider::GetNavigationTool() const
+{
+	if (const TSharedPtr<IAvaSequencer> AvaSequencer = WeakAvaSequencer.Pin())
+	{
+		return SequenceNavigator::FNavigationToolExtender::FindNavigationTool(AvaSequencer->GetSequencer());
+	}
+	return nullptr;
+}
+
+IAvaSceneInterface* FAvaNavigationToolProvider::GetSceneInterface(const UE::SequenceNavigator::INavigationTool& InTool) const
+{
+	if (const TSharedPtr<ISequencer> Sequencer = InTool.GetSequencer())
+	{
+		return FAvaSequencerUtils::GetSceneInterface(Sequencer.ToSharedRef());
+	}
+	return nullptr;
+}
+
+IAvaSequenceProvider* FAvaNavigationToolProvider::GetSequenceProvider(const UE::SequenceNavigator::INavigationTool& InTool) const
+{
+	if (const TSharedPtr<ISequencer> Sequencer = InTool.GetSequencer())
+	{
+		return FAvaSequencerUtils::GetSequenceProvider(Sequencer.ToSharedRef());
+	}
+	return nullptr;
+}
+
+const IAvaSequencerProvider* FAvaNavigationToolProvider::GetSequencerProvider() const
+{
+	if (const TSharedPtr<IAvaSequencer> AvaSequencer = WeakAvaSequencer.Pin())
+	{
+		return &AvaSequencer->GetProvider();
+	}
+	return nullptr;
+}
+
+IAvaSequencePlaybackObject* FAvaNavigationToolProvider::GetSequencerPlaybackObject() const
+{
+	if (const IAvaSequencerProvider* const SequencerProvider = GetSequencerProvider())
+	{
+		return SequencerProvider->GetPlaybackObject();
+	}
+	return nullptr;
+}
+
+bool FAvaNavigationToolProvider::CanEditOrPlaySelection(const int32 InMinNumSelected, const int32 InMaxNumSelected) const
+{
+	if (const IAvaSequencerProvider* const SequencerProvider = GetSequencerProvider())
+	{
+		if (SequencerProvider->CanEditOrPlaySequences())
+		{
+			const int32 NumSelected = GetSelectedSequenceItems().Num();
+			return NumSelected >= InMinNumSelected
+				&& (InMaxNumSelected == INDEX_NONE || NumSelected <= InMaxNumSelected);
+		}
+	}
+	return false;
+}
+
+void FAvaNavigationToolProvider::ExtendToolToolBar()
+{
+	UToolMenu* const ToolMenu = UToolMenus::Get()->ExtendMenu(UE::SequenceNavigator::GetToolBarMenuName());
+	if (!ToolMenu)
+	{
+		return;
+	}
+
+	FToolMenuSection& MotionDesignSection = ToolMenu->FindOrAddSection(ToolbarSectionName);
+
+	const FAvaSequenceNavigatorCommands& AvaSequenceNavigatorCommands = FAvaSequenceNavigatorCommands::Get();
+	const FName SequencerToolbarStyleName = TEXT("SequencerToolbar");
+
+	FSlimHorizontalToolBarBuilder ToolBarBuilder(ToolCommands, FMultiBoxCustomization::None, nullptr, true);
+	ToolBarBuilder.SetLabelVisibility(EVisibility::Collapsed);
+
+	FToolMenuEntry SequenceAddNewEntry = FToolMenuEntry::InitToolBarButton(AvaSequenceNavigatorCommands.AddNew
+		, TAttribute<FText>(), TAttribute<FText>()
+		, FSlateIcon(FAppStyle::GetAppStyleSetName(), TEXT("AnimationEditor.SetKey")));
+	SequenceAddNewEntry.StyleNameOverride = SequencerToolbarStyleName;
+	SequenceAddNewEntry.InsertPosition = FToolMenuInsert(TEXT("AddFilter"), EToolMenuInsertType::Before);
+	SequenceAddNewEntry.SetCommandList(ToolCommands);
+	MotionDesignSection.AddEntry(SequenceAddNewEntry);
+
+	FToolMenuEntry SequencePlaySelectedEntry = FToolMenuEntry::InitToolBarButton(AvaSequenceNavigatorCommands.PlaySelected
+		, TAttribute<FText>(), TAttribute<FText>()
+		, FSlateIcon(FAppStyle::GetAppStyleSetName(), TEXT("Icons.Toolbar.Play")));
+	SequencePlaySelectedEntry.StyleNameOverride = SequencerToolbarStyleName;
+	SequencePlaySelectedEntry.SetCommandList(ToolCommands);
+	MotionDesignSection.AddEntry(SequencePlaySelectedEntry);
+
+	FToolMenuEntry SequenceContinueSelectedEntry = FToolMenuEntry::InitToolBarButton(AvaSequenceNavigatorCommands.ContinueSelected
+		, TAttribute<FText>(), TAttribute<FText>()
+		, FSlateIcon(FAppStyle::GetAppStyleSetName(), TEXT("Icons.JumpToEvent")));
+	SequenceContinueSelectedEntry.StyleNameOverride = SequencerToolbarStyleName;
+	SequenceContinueSelectedEntry.SetCommandList(ToolCommands);
+	MotionDesignSection.AddEntry(SequenceContinueSelectedEntry);
+
+	FToolMenuEntry SequenceStopSelectedEntry = FToolMenuEntry::InitToolBarButton(AvaSequenceNavigatorCommands.StopSelected
+		, TAttribute<FText>(), TAttribute<FText>()
+		, FSlateIcon(FAppStyle::GetAppStyleSetName(), TEXT("Icons.Toolbar.Stop")));
+	SequenceStopSelectedEntry.StyleNameOverride = SequencerToolbarStyleName;
+	SequenceStopSelectedEntry.SetCommandList(ToolCommands);
+	MotionDesignSection.AddEntry(SequenceStopSelectedEntry);
+}
+
+void FAvaNavigationToolProvider::ExtendToolItemContextMenu()
+{
+	UToolMenus* ToolMenus = UToolMenus::Get();
+	if (!ToolMenus)
+	{
+		return;
+	}
+
+	UToolMenu* const ExtendedMenu = ToolMenus->ExtendMenu(UE::SequenceNavigator::GetItemContextMenuName());
+	if (!ExtendedMenu)
+	{
+		return;
+	}
+
+	const FAvaSequenceNavigatorCommands& AvaSequenceNavigatorCommands = FAvaSequenceNavigatorCommands::Get();
+
+	FToolMenuSection& MotionDesignSection = ExtendedMenu->FindOrAddSection(ContextMenuSectionName
+		, LOCTEXT("MotionDesign", "Motion Design Actions")
+		, FToolMenuInsert(TEXT("ToolActions"), EToolMenuInsertType::Before));
+
+	MotionDesignSection.AddSubMenu(TEXT("ApplyPreset")
+		, LOCTEXT("ApplyPresetLabel", "Apply Preset")
+		, LOCTEXT("ApplyPresetTooltip", "Apply a Preset to the Selected Sequences")
+		, FNewToolMenuDelegate::CreateSP(this, &FAvaNavigationToolProvider::GeneratePresetMenu));
+	MotionDesignSection.AddMenuEntry(AvaSequenceNavigatorCommands.SpawnSequencePlayer);
+	MotionDesignSection.AddMenuEntry(AvaSequenceNavigatorCommands.ExportSequence);
+
+	const FGenericCommands& GenericCommands = FGenericCommands::Get();
+	FToolMenuSection& GenericSection = ExtendedMenu->FindOrAddSection(TEXT("GenericActions"), LOCTEXT("GenericActionsHeader", "Generic Actions"));
+
+	GenericSection.AddMenuEntry(GenericCommands.Duplicate);
+	GenericSection.AddMenuEntry(GenericCommands.Delete);
+	GenericSection.AddMenuEntry(GenericCommands.Rename);
+}
+
+void FAvaNavigationToolProvider::RemoveToolToolBarExtension()
+{
+	if (UToolMenu* const Menu = UToolMenus::Get()->FindMenu(UE::SequenceNavigator::GetToolBarMenuName()))
+	{
+		Menu->RemoveSection(ToolbarSectionName);
+	}
+}
+
+void FAvaNavigationToolProvider::RemoveToolItemContextMenuExtension()
+{
+	if (UToolMenu* const Menu = UToolMenus::Get()->FindMenu(UE::SequenceNavigator::GetItemContextMenuName()))
+	{
+		Menu->RemoveSection(ContextMenuSectionName);
+	}
+}
+
+void FAvaNavigationToolProvider::ExtendToolItemDragDropOp(UE::SequenceNavigator::FNavigationToolItemDragDropOp& InDragDropOp)
+{
+	InDragDropOp.AddDropHandler<UE::SequenceNavigator::FNavigationToolAvaSequenceDropHandler>(WeakAvaSequencer);
+}
+
+void FAvaNavigationToolProvider::OnSequenceAdded(UAvaSequence* const InAvaSequence)
+{
+	using namespace UE::SequenceNavigator;
+
+	const TSharedPtr<IAvaSequencer> AvaSequencer = WeakAvaSequencer.Pin();
+	if (!AvaSequencer.IsValid())
+	{
+		return;
+	}
+
+	const TSharedPtr<ISequencer> Sequencer = AvaSequencer->GetSequencer();
+	if (!Sequencer.IsValid())
+	{
+		return;
+	}
+
+	const TSharedPtr<INavigationTool> NavigationTool = FNavigationToolExtender::FindNavigationTool(Sequencer.ToSharedRef());
+	if (!NavigationTool.IsValid())
+	{
+		return;
+	}
+
+	NavigationTool->RequestRefresh();
+}
+
+void FAvaNavigationToolProvider::GeneratePresetMenu(UToolMenu* const InToolMenu)
+{
+	const UAvaSequencerSettings* const SequencerSettings = GetDefault<UAvaSequencerSettings>();
+	if (!SequencerSettings)
+	{
+		return;
+	}
+
+	const TSharedPtr<IAvaSequencer> AvaSequencer = WeakAvaSequencer.Pin();
+	if (!AvaSequencer.IsValid())
+	{
+		return;
+	}
+
+	const TSharedRef<IAvaSequencer> AvaSequencerRef = AvaSequencer.ToSharedRef();
+
+	// Default Presets
+	const TConstArrayView<FAvaSequencePreset> DefaultPresets = SequencerSettings->GetDefaultSequencePresets();
+	if (!DefaultPresets.IsEmpty())
+	{
+		FToolMenuSection& DefaultPresetSection = InToolMenu->FindOrAddSection(TEXT("DefaultPresets"), LOCTEXT("DefaultPresetsLabel", "Default Presets"));
+		for (const FAvaSequencePreset& Preset : DefaultPresets)
+		{
+			DefaultPresetSection.AddMenuEntry(Preset.PresetName
+				, FText::FromName(Preset.PresetName)
+				, FText::FromName(Preset.PresetName)
+				, FSlateIcon()
+				, FToolUIActionChoice(FExecuteAction::CreateSP(this, &FAvaNavigationToolProvider::ApplyDefaultPresetToSelection, Preset.PresetName)));
+		}
+	}
+
+	// Custom Presets
+	const TSet<FAvaSequencePreset>& CustomPresets = SequencerSettings->GetCustomSequencePresets();
+	if (!CustomPresets.IsEmpty())
+	{
+		FToolMenuSection& CustomPresetSection = InToolMenu->FindOrAddSection(TEXT("CustomPresets"), LOCTEXT("CustomPresetsLabel", "Custom Presets"));
+		for (const FAvaSequencePreset& Preset : CustomPresets)
+		{
+			CustomPresetSection.AddMenuEntry(Preset.PresetName
+				, FText::FromName(Preset.PresetName)
+				, FText::FromName(Preset.PresetName)
+				, FSlateIcon()
+				, FToolUIActionChoice(FExecuteAction::CreateSP(this, &FAvaNavigationToolProvider::ApplyCustomPresetToSelection, Preset.PresetName)));
+		}
+	}
+
+	// Settings
+	FToolMenuSection& SettingsSection = InToolMenu->FindOrAddSection(TEXT("Settings"));
+	SettingsSection.AddSeparator(NAME_None);
+	SettingsSection.AddMenuEntry(TEXT("OpenSettings")
+		, LOCTEXT("OpenSettingsLabel", "Open Settings")
+		, LOCTEXT("OpenSettingsTooltip", "Opens the Settings to customize the sequence presets")
+		, FSlateIcon(FAppStyle::GetAppStyleSetName(), "Icons.Settings")
+		, FExecuteAction::CreateLambda([]
+		{
+			if (const UAvaSequencerSettings* const AvaSequencerSettings = GetDefault<UAvaSequencerSettings>())
+			{
+				ISettingsModule& SettingsModule = FModuleManager::LoadModuleChecked<ISettingsModule>(TEXT("Settings"));
+				SettingsModule.ShowViewer(AvaSequencerSettings->GetContainerName(), AvaSequencerSettings->GetCategoryName(), AvaSequencerSettings->GetSectionName());
+			}
+		}));
+}
+
+TArray<TSharedPtr<UE::SequenceNavigator::FNavigationToolAvaSequence>> FAvaNavigationToolProvider::GetSelectedSequenceItems() const
+{
+	using namespace UE::SequenceNavigator;
+
+	const TSharedPtr<INavigationTool> Tool = GetNavigationTool();
+	if (!Tool.IsValid())
+	{
+		return {};
+	}
+
+	const TArray<FNavigationToolItemPtr> SelectedItems = Tool->GetSelectedItems();
+	if (SelectedItems.IsEmpty())
+	{
+		return {};
+	}
+
+	TArray<TSharedPtr<FNavigationToolAvaSequence>> SelectedSequenceItems;
+	SelectedSequenceItems.Reserve(SelectedItems.Num());
+
+	for (const FNavigationToolItemPtr& Item : SelectedItems)
+	{
+		if (FNavigationToolAvaSequence* const AvaSequenceItem = Item->CastTo<FNavigationToolAvaSequence>())
+		{
+			SelectedSequenceItems.Add(SharedThis(AvaSequenceItem));
+		}
+	}
+
+	return SelectedSequenceItems;
+}
+
+TArray<UAvaSequence*> FAvaNavigationToolProvider::GetSelectedSequences() const
+{
+	using namespace UE::SequenceNavigator;
+
+	const TArray<TSharedPtr<FNavigationToolAvaSequence>> SelectedItems = GetSelectedSequenceItems();
+	if (SelectedItems.IsEmpty())
+	{
+		return {};
+	}
+
+	TArray<UAvaSequence*> SelectedAvaSequences;
+	SelectedAvaSequences.Reserve(SelectedItems.Num());
+
+	for (const TSharedPtr<FNavigationToolAvaSequence>& SequenceItem : SelectedItems)
+	{
+		if (UAvaSequence* const AvaSequence = SequenceItem->GetAvaSequence())
+		{
+			SelectedAvaSequences.Add(AvaSequence);
+		}
+	}
+
+	return SelectedAvaSequences;
+}
+
+bool FAvaNavigationToolProvider::CanRelabelSelection() const
+{
+	return CanEditOrPlaySelection(/*InMinNumSelected=*/1, /*InMaxNumSelected=*/1);
+}
+
+void FAvaNavigationToolProvider::RelabelSelection()
+{
+	using namespace UE::SequenceNavigator;
+
+	const TSharedPtr<IAvaSequencer> AvaSequencer = WeakAvaSequencer.Pin();
+	if (!AvaSequencer.IsValid())
+	{
+		return;
+	}
+
+	const TSharedPtr<INavigationTool> NavigationTool = FNavigationToolExtender::FindNavigationTool(AvaSequencer->GetSequencer());
+	if (!NavigationTool.IsValid())
+	{
+		return;
+	}
+
+	const TArray<FNavigationToolItemPtr> SelectedItems = NavigationTool->GetSelectedItems();
+	if (SelectedItems.Num() != 1)
+	{
+		return;
+	}
+
+	const FNavigationToolAvaSequence* const AvaSequenceItem = SelectedItems[0]->CastTo<FNavigationToolAvaSequence>();
+	if (!AvaSequenceItem)
+	{
+		return;
+	}
+
+	// @Todo: Implement BeginRename function in base classes
+	//AvaSequenceItem->BeginRename();
+}
+
+bool FAvaNavigationToolProvider::CanAddSequenceToSelection() const
+{
+	return CanEditOrPlaySelection(/*InMinNumSelected=*/0, /*InMaxNumSelected=*/1);
+}
+
+void FAvaNavigationToolProvider::AddSequenceToSelection()
+{
+	using namespace UE::SequenceNavigator;
+
+	const TSharedPtr<IAvaSequencer> AvaSequencer = WeakAvaSequencer.Pin();
+	if (!AvaSequencer.IsValid())
+	{
+		return;
+	}
+
+	const TSharedPtr<INavigationTool> NavigationTool = FNavigationToolExtender::FindNavigationTool(AvaSequencer->GetSequencer());
+	if (!NavigationTool.IsValid())
+	{
+		return;
+	}
+
+	TSharedPtr<FNavigationToolAvaSequence> ParentSequenceItem;
+	UAvaSequence* ParentSequence = nullptr;
+
+	const TArray<TSharedPtr<FNavigationToolAvaSequence>> SelectedSequences = GetSelectedSequenceItems();
+	if (SelectedSequences.Num() == 1)
+	{
+		ParentSequenceItem = SelectedSequences[0];
+		ParentSequence = ParentSequenceItem->GetAvaSequence();
+	}
+
+	UAvaSequence* const NewSequence = AvaSequencer->AddSequence(ParentSequence);
+	if (!NewSequence)
+	{
+		return;
+	}
+
+	FNavigationToolAddItemParams AddChildItemParams;
+	AddChildItemParams.Item = NavigationTool->FindOrAdd<FNavigationToolAvaSequence>(SharedThis(this)
+		, AddChildItemParams.RelativeItem, NewSequence);
+	AddChildItemParams.RelativeItem = ParentSequenceItem.IsValid()
+		? ParentSequenceItem.ToSharedRef() : NavigationTool->GetTreeRoot();
+	AddChildItemParams.RelativeDropZone = EItemDropZone::OntoItem;
+	AddChildItemParams.Flags = ENavigationToolAddItemFlags::Select | ENavigationToolAddItemFlags::Transact;
+	AddChildItemParams.SelectionFlags = ENavigationToolItemSelectionFlags::SignalSelectionChange
+		| ENavigationToolItemSelectionFlags::ScrollIntoView;
+
+	NavigationTool->EnqueueItemAction<FNavigationToolAddItem>(AddChildItemParams);
+}
+
+bool FAvaNavigationToolProvider::CanDuplicateSelection() const
+{
+	return CanEditOrPlaySelection();
+}
+
+void FAvaNavigationToolProvider::DuplicateSelection()
+{
+	const IAvaSequencerProvider* const SequencerProvider = GetSequencerProvider();
+	if (!SequencerProvider)
+	{
+		return;
+	}
+
+	IAvaSequenceProvider* const SequenceProvider = SequencerProvider->GetSequenceProvider();
+	if (!SequenceProvider)
+	{
+		return;
+	}
+
+	const TArray<UAvaSequence*> SelectedSequences = GetSelectedSequences();
+	if (SelectedSequences.IsEmpty())
+	{
+		return;
+	}
+
+	UObject* const Outer = SequenceProvider->ToUObject();
+	if (!Outer)
+	{
+		return;
+	}
+
+	const FScopedTransaction Transaction(LOCTEXT("DuplicateSequenceTransaction", "Duplicate Sequence"));
+
+	Outer->Modify();
+
+	for (UAvaSequence* const TemplateSequence : SelectedSequences)
+	{
+		if (TemplateSequence)
+		{
+			UAvaSequence* const DupedSequence = DuplicateObject<UAvaSequence>(TemplateSequence, Outer);
+			SequenceProvider->AddSequence(DupedSequence);
+		}
+	}
+}
+
+bool FAvaNavigationToolProvider::CanDeleteSelection() const
+{
+	return CanEditOrPlaySelection();
+}
+
+void FAvaNavigationToolProvider::DeleteSelection()
+{
+	const TSharedPtr<IAvaSequencer> AvaSequencer = WeakAvaSequencer.Pin();
+	if (!AvaSequencer.IsValid())
+	{
+		return;
+	}
+
+	AvaSequencer->DeleteSequences(TSet<UAvaSequence*> { GetSelectedSequences() });
+}
+
+bool FAvaNavigationToolProvider::CanPlaySelection() const
+{
+	return CanEditOrPlaySelection();
+}
+
+void FAvaNavigationToolProvider::PlaySelection()
+{
+	IAvaSequencePlaybackObject* const PlaybackObject = GetSequencerPlaybackObject();
+	if (!PlaybackObject)
+	{
+		return;
+	}
+
+	FAvaSequencePlayParams PlaySettings;
+	PlaySettings.AdvancedSettings.bRestoreState = true;
+
+	const TArray<UAvaSequence*> SelectedSequences = GetSelectedSequences();
+	for (UAvaSequence* const Sequence : SelectedSequences)
+	{
+		if (Sequence)
+		{
+			PlaybackObject->PlaySequence(Sequence, PlaySettings);
+		}
+	}
+}
+
+bool FAvaNavigationToolProvider::CanContinueSelection() const
+{
+	return CanEditOrPlaySelection();
+}
+
+void FAvaNavigationToolProvider::ContinueSelection()
+{
+	IAvaSequencePlaybackObject* const PlaybackObject = GetSequencerPlaybackObject();
+	if (!PlaybackObject)
+	{
+		return;
+	}
+
+	const TArray<UAvaSequence*> SelectedSequences = GetSelectedSequences();
+	for (UAvaSequence* const Sequence : SelectedSequences)
+	{
+		if (Sequence)
+		{
+			PlaybackObject->ContinueSequence(Sequence);
+		}
+	}
+}
+
+bool FAvaNavigationToolProvider::CanStopSelection() const
+{
+	return CanEditOrPlaySelection();
+}
+
+void FAvaNavigationToolProvider::StopSelection()
+{
+	IAvaSequencePlaybackObject* const PlaybackObject = GetSequencerPlaybackObject();
+	if (!PlaybackObject)
+	{
+		return;
+	}
+
+	const TArray<UAvaSequence*> SelectedSequences = GetSelectedSequences();
+	for (UAvaSequence* const Sequence : SelectedSequences)
+	{
+		if (Sequence)
+		{
+			PlaybackObject->StopSequence(Sequence);
+		}
+	}
+}
+
+bool FAvaNavigationToolProvider::CanExportSelection() const
+{
+	return CanEditOrPlaySelection();
+}
+
+void FAvaNavigationToolProvider::ExportSelection()
+{
+	IAvaSequencerProvider* const SequencerProvider = const_cast<IAvaSequencerProvider*>(GetSequencerProvider());
+	if (!SequencerProvider)
+	{
+		return;
+	}
+
+	const TArray<UAvaSequence*> SelectedSequences = GetSelectedSequences();
+	if (SelectedSequences.IsEmpty())
+	{
+		return;
+	}
+
+	TArray<TWeakObjectPtr<>> WeakSelectedObjects;
+	Algo::Transform(SelectedSequences, WeakSelectedObjects, [](UAvaSequence* const InSequence)
+		{
+			return InSequence;
+		});
+
+	SequencerProvider->ExportSequences(SelectedSequences);
+
+	FNotificationInfo Info(LOCTEXT("ExportSuccess", "Sequence Exported Successfully!"));
+	Info.bUseLargeFont = false;
+	Info.bUseSuccessFailIcons = false;
+	Info.ExpireDuration = 5.0f;
+	Info.HyperlinkText = LOCTEXT("ShowNewAssetsInContentBrowser", "Show in content browser");
+	Info.Hyperlink = FSimpleDelegate::CreateLambda([WeakSelectedObjects]()
+		{
+			TArray<UObject*> SyncObjects;
+			SyncObjects.Reserve(WeakSelectedObjects.Num());
+
+			for (const TWeakObjectPtr<> WeakObject : WeakSelectedObjects)
+			{
+				if (UObject* const SyncObject = WeakObject.Get())
+				{
+					SyncObjects.Add(SyncObject);
+				}
+			}
+
+			FContentBrowserModule& ContentBrowserModule = FModuleManager::LoadModuleChecked<FContentBrowserModule>(TEXT("ContentBrowser"));
+			ContentBrowserModule.Get().SyncBrowserToAssets(SyncObjects);
+		});
+
+	FSlateNotificationManager::Get().AddNotification(Info);
+}
+
+bool FAvaNavigationToolProvider::CanSpawnPlayersForSelection() const
+{
+	return CanEditOrPlaySelection();
+}
+
+void FAvaNavigationToolProvider::SpawnPlayersForSelection()
+{
+	if (!GEditor)
+	{
+		return;
+	}
+
+	IAvaSequencePlaybackObject* const PlaybackObject = GetSequencerPlaybackObject();
+	if (!PlaybackObject)
+	{
+		return;
+	}
+
+	const UObject* const PlaybackContext = PlaybackObject->GetPlaybackContext();
+	if (!PlaybackContext)
+	{
+		return;
+	}
+
+	UWorld* const World = PlaybackContext->GetWorld();
+	if (!World)
+	{
+		return;
+	}
+
+	UActorFactory* const ActorFactory = GEditor->FindActorFactoryForActorClass(AAvaSequenceActor::StaticClass());
+	if (!ActorFactory)
+	{
+		return;
+	}
+
+	const TArray<UAvaSequence*> Sequences = GetSelectedSequences();
+	if (Sequences.IsEmpty())
+	{
+		return;
+	}
+
+	const FScopedTransaction Transaction(LOCTEXT("SpawnSequencePlayers", "Spawn Sequence Players"));
+
+	for (UAvaSequence* const Sequence : Sequences)
+	{
+		if (Sequence)
+		{
+			GEditor->UseActorFactory(ActorFactory, FAssetData(Sequence), &FTransform::Identity);
+		}
+	}
+
+	const FText NotifyText = FText::Format(LOCTEXT("SpawnPlayerSuccess", "{0} Sequence Players Spawned Successfully!")
+		, FText::AsNumber(Sequences.Num()));
+
+	FNotificationInfo Info(NotifyText);
+	Info.bUseLargeFont = false;
+	Info.bUseSuccessFailIcons = false;
+	FSlateNotificationManager::Get().AddNotification(Info);
+}
+
+void FAvaNavigationToolProvider::ApplyDefaultPresetToSelection(FName InPresetName)
+{
+	const UAvaSequencerSettings* SequencerSettings = GetDefault<UAvaSequencerSettings>();
+	if (!SequencerSettings)
+	{
+		return;
+	}
+
+	const TConstArrayView<FAvaSequencePreset> DefaultSequencePresets = SequencerSettings->GetDefaultSequencePresets();
+
+	const int32 PresetIndex = DefaultSequencePresets.Find(FAvaSequencePreset(InPresetName));
+	if (PresetIndex == INDEX_NONE)
+	{
+		return;
+	}
+
+	ApplyPresetToSelection(DefaultSequencePresets[PresetIndex]);
+}
+
+void FAvaNavigationToolProvider::ApplyCustomPresetToSelection(FName InPresetName)
+{
+	const UAvaSequencerSettings* SequencerSettings = GetDefault<UAvaSequencerSettings>();
+	if (!SequencerSettings)
+	{
+		return;
+	}
+
+	const FAvaSequencePreset* SequencePreset = SequencerSettings->GetCustomSequencePresets().Find(FAvaSequencePreset(InPresetName));
+	if (!SequencePreset)
+	{
+		return;
+	}
+
+	ApplyPresetToSelection(*SequencePreset);
+}
+
+void FAvaNavigationToolProvider::ApplyPresetToSelection(const FAvaSequencePreset& InPreset)
+{
+	const TArray<UAvaSequence*> SelectedSequences = GetSelectedSequences();
+	if (SelectedSequences.IsEmpty())
+	{
+		return;
+	}
+
+	FScopedTransaction Transaction(LOCTEXT("ApplySequencePreset", "Apply Sequence Preset"));
+
+	for (UAvaSequence* const AvaSequence : SelectedSequences)
+	{
+		if (AvaSequence)
+		{
+			InPreset.ApplyPreset(AvaSequence);
+		}
+	}
+}
+
+} // namespace UE::AvaSequencer
+
+#undef LOCTEXT_NAMESPACE
diff --git a/Engine/Plugins/Experimental/AvalancheSequenceNavigator/Source/AvalancheSequenceNavigator/Private/AvaNavigationToolProvider.h b/Engine/Plugins/Experimental/AvalancheSequenceNavigator/Source/AvalancheSequenceNavigator/Private/AvaNavigationToolProvider.h
new file mode 100644
--- /dev/null
+++ b/Engine/Plugins/Experimental/AvalancheSequenceNavigator/Source/AvalancheSequenceNavigator/Private/AvaNavigationToolProvider.h
@@ -0,0 +1,135 @@
+﻿// Copyright Epic Games, Inc. All Rights Reserved.
+
+#pragma once
+
+#include "Delegates/IDelegateInstance.h"
+#include "NavigationToolDefines.h"
+#include "Providers/NavigationToolProvider.h"
+#include "Templates/SharedPointer.h"
+#include "UObject/WeakObjectPtr.h"
+
+class FUICommandList;
+class IAvaSceneInterface;
+class IAvaSequencePlaybackObject;
+class IAvaSequenceProvider;
+class IAvaSequencerProvider;
+class IAvaSequencer;
+class UAvaSequence;
+class UMovieSceneSequence;
+class UObject;
+struct FAvaSequencePreset;
+struct FNavigationToolSaveState;
+
+namespace UE::SequenceNavigator
+{
+	class FNavigationToolAvaSequence;
+	class FNavigationToolItem;
+	class FNavigationToolItemDragDropOp;
+	class INavigationTool;
+}
+
+namespace UE::AvaSequencer
+{
+
+class FAvaNavigationToolProvider : public SequenceNavigator::FNavigationToolProvider
+{
+public:
+	static const FName Identifier;
+
+	static const FText MotionDesignColumnViewName;
+
+	FAvaNavigationToolProvider(const TSharedRef<IAvaSequencer>& InAvaSequencer);
+	virtual ~FAvaNavigationToolProvider() override = default;
+
+	//~ Begin INavigationToolProvider
+
+	virtual FName GetIdentifier() const override;
+
+	virtual TSet<TSubclassOf<UMovieSceneSequence>> GetSupportedSequenceClasses() const override;
+
+	virtual FText GetDefaultColumnView() const override;
+
+	virtual FNavigationToolSaveState* GetSaveState(const SequenceNavigator::INavigationTool& InTool) const override;
+	virtual void SetSaveState(const SequenceNavigator::INavigationTool& InTool, const FNavigationToolSaveState& InSaveState) const override;
+
+	virtual void BindCommands(const TSharedRef<FUICommandList>& InCommandList) override;
+
+	virtual void OnActivate() override;
+	virtual void OnDeactivate() override;
+
+	virtual void OnExtendColumns(SequenceNavigator::FNavigationToolColumnExtender& OutExtender) override;
+	virtual void OnExtendColumnViews(TSet<FNavigationToolColumnView>& OutColumnViews) override;
+	virtual void OnExtendItemChildren(SequenceNavigator::INavigationTool& InTool
+		, const SequenceNavigator::FNavigationToolItemPtr& InParentItem
+		, TArray<SequenceNavigator::FNavigationToolItemPtr>& OutChildren
+		, const bool bInRecursive) override;
+	virtual void OnExtendBuiltInFilters(TArray<FNavigationToolBuiltInFilterParams>& OutFilterParams) override;
+
+	//~ End INavigationToolProvider
+
+private:
+	static const FName ToolbarSectionName;
+	static const FName ContextMenuSectionName;
+
+	TSharedPtr<SequenceNavigator::INavigationTool> GetNavigationTool() const;
+	IAvaSceneInterface* GetSceneInterface(const SequenceNavigator::INavigationTool& InTool) const;
+	IAvaSequenceProvider* GetSequenceProvider(const SequenceNavigator::INavigationTool& InTool) const;
+	const IAvaSequencerProvider* GetSequencerProvider() const;
+	IAvaSequencePlaybackObject* GetSequencerPlaybackObject() const;
+
+	bool CanEditOrPlaySelection(const int32 InMinNumSelected = 1, const int32 InMaxNumSelected = INDEX_NONE) const;
+
+	void ExtendToolToolBar();
+	void ExtendToolItemContextMenu();
+
+	void RemoveToolToolBarExtension();
+	void RemoveToolItemContextMenuExtension();
+
+	void ExtendToolItemDragDropOp(SequenceNavigator::FNavigationToolItemDragDropOp& InDragDropOp);
+
+	void OnSequenceAdded(UAvaSequence* const InAvaSequence);
+
+	void GeneratePresetMenu(UToolMenu* const InToolMenu);
+
+	TArray<TSharedPtr<SequenceNavigator::FNavigationToolAvaSequence>> GetSelectedSequenceItems() const;
+	TArray<UAvaSequence*> GetSelectedSequences() const;
+
+	bool CanRelabelSelection() const;
+	void RelabelSelection();
+
+	bool CanAddSequenceToSelection() const;
+	void AddSequenceToSelection();
+
+	bool CanDuplicateSelection() const;
+	void DuplicateSelection();
+
+	bool CanDeleteSelection() const;
+	void DeleteSelection();
+
+	bool CanPlaySelection() const;
+	void PlaySelection();
+
+	bool CanContinueSelection() const;
+	void ContinueSelection();
+
+	bool CanStopSelection() const;
+	void StopSelection();
+
+	bool CanExportSelection() const;
+	void ExportSelection();
+
+	bool CanSpawnPlayersForSelection() const;
+	void SpawnPlayersForSelection();
+
+	void ApplyDefaultPresetToSelection(const FName InPresetName);
+	void ApplyCustomPresetToSelection(const FName InPresetName);
+	void ApplyPresetToSelection(const FAvaSequencePreset& InPreset);
+
+	TWeakPtr<IAvaSequencer> WeakAvaSequencer;
+
+	TSharedRef<FUICommandList> ToolCommands;
+
+	FDelegateHandle DragDropInitializedDelegate;
+};
+
+} // namespace UE::AvaSequencer
diff --git a/Engine/Plugins/Experimental/AvalancheSequenceNavigator/Source/AvalancheSequenceNavigator/Private/AvaSequenceNavigatorCommands.cpp b/Engine/Plugins/Experimental/AvalancheSequenceNavigator/Source/AvalancheSequenceNavigator/Private/AvaSequenceNavigatorCommands.cpp
new file mode 100644
--- /dev/null
+++ b/Engine/Plugins/Experimental/AvalancheSequenceNavigator/Source/AvalancheSequenceNavigator/Private/AvaSequenceNavigatorCommands.cpp
@@ -0,0 +1,46 @@
+﻿// Copyright Epic Games, Inc. All Rights Reserved.
+
+#include "AvaSequenceNavigatorCommands.h"
+
+#define LOCTEXT_NAMESPACE "AvaSequenceNavigatorCommands"
+
+void FAvaSequenceNavigatorCommands::RegisterCommands()
+{
+	UI_COMMAND(AddNew
+		, "Add New"
+		, "Adds a new embedded Sequence to the level"
+		, EUserInterfaceActionType::Button
+		, FInputChord());
+
+	UI_COMMAND(PlaySelected
+		, "Play"
+		, "Plays all selected Sequences"
+		, EUserInterfaceActionType::Button
+		, FInputChord());
+
+	UI_COMMAND(ContinueSelected
+		, "Continue"
+		, "Continues all selected Sequences"
+		, EUserInterfaceActionType::Button
+		, FInputChord());
+
+	UI_COMMAND(StopSelected
+		, "Stop"
+		, "Stops all selected Sequences"
+		, EUserInterfaceActionType::Button
+		, FInputChord());
+
+	UI_COMMAND(ExportSequence
+		, "Export"
+		, "Exports the Sequence into its own Asset"
+		, EUserInterfaceActionType::Button
+		, FInputChord());
+
+	UI_COMMAND(SpawnSequencePlayer
+		, "Spawn Sequence Player"
+		, "Spawns a Sequence Player for the selected Sequence"
+		, EUserInterfaceActionType::Button
+		, FInputChord());
+}
+
+#undef LOCTEXT_NAMESPACE
diff --git a/Engine/Plugins/Experimental/AvalancheSequenceNavigator/Source/AvalancheSequenceNavigator/Private/AvaSequenceNavigatorCommands.h b/Engine/Plugins/Experimental/AvalancheSequenceNavigator/Source/AvalancheSequenceNavigator/Private/AvaSequenceNavigatorCommands.h
new file mode 100644
--- /dev/null
+++ b/Engine/Plugins/Experimental/AvalancheSequenceNavigator/Source/AvalancheSequenceNavigator/Private/AvaSequenceNavigatorCommands.h
@@ -0,0 +1,32 @@
+﻿// Copyright Epic Games, Inc. All Rights Reserved.
+
+#pragma once
+
+#include "Framework/Commands/Commands.h"
+#include "Styling/AppStyle.h"
+
+class FAvaSequenceNavigatorCommands : public TCommands<FAvaSequenceNavigatorCommands>
+{
+public:
+	FAvaSequenceNavigatorCommands()
+		: TCommands<FAvaSequenceNavigatorCommands>(TEXT("AvaSequenceNavigatorCommands")
+		, NSLOCTEXT("MotionDesignSequenceNavigatorCommands", "MotionDesignSequenceNavigatorCommands", "Motion Design Sequence Navigator Commands")
+		, NAME_None
+		, FAppStyle::GetAppStyleSetName())
+	{
+	}
+
+	virtual void RegisterCommands() override;
+
+	TSharedPtr<FUICommandInfo> AddNew;
+
+	TSharedPtr<FUICommandInfo> PlaySelected;
+
+	TSharedPtr<FUICommandInfo> ContinueSelected;
+
+	TSharedPtr<FUICommandInfo> StopSelected;
+
+	TSharedPtr<FUICommandInfo> ExportSequence;
+
+	TSharedPtr<FUICommandInfo> SpawnSequencePlayer;
+};
diff --git a/Engine/Plugins/Experimental/AvalancheSequenceNavigator/Source/AvalancheSequenceNavigator/Private/AvalancheSequenceNavigatorModule.cpp b/Engine/Plugins/Experimental/AvalancheSequenceNavigator/Source/AvalancheSequenceNavigator/Private/AvalancheSequenceNavigatorModule.cpp
new file mode 100644
--- /dev/null
+++ b/Engine/Plugins/Experimental/AvalancheSequenceNavigator/Source/AvalancheSequenceNavigator/Private/AvalancheSequenceNavigatorModule.cpp
@@ -0,0 +1,100 @@
+﻿// Copyright Epic Games, Inc. All Rights Reserved.
+
+#include "AvalancheSequenceNavigatorModule.h"
+#include "AvaNavigationToolProvider.h"
+#include "AvaSequenceNavigatorCommands.h"
+#include "AvaSequencerSubsystem.h"
+#include "IAvaSequencer.h"
+#include "ISequencer.h"
+#include "ISequencerModule.h"
+#include "Modules/ModuleManager.h"
+#include "NavigationToolExtender.h"
+
+#define LOCTEXT_NAMESPACE "AvalancheSequenceNavigatorModule"
+
+using namespace UE::AvaSequencer;
+using namespace UE::SequenceNavigator;
+
+void FAvalancheSequenceNavigatorModule::StartupModule()
+{
+	FAvaSequenceNavigatorCommands::Register();
+
+	if (ISequencerModule* const SequencerModule = FModuleManager::LoadModulePtr<ISequencerModule>(TEXT("Sequencer")))
+	{
+		SequencerCreatedHandle = SequencerModule->RegisterOnSequencerCreated(
+			FOnSequencerCreated::FDelegate::CreateRaw(this, &FAvalancheSequenceNavigatorModule::OnSequencerCreated));
+	}
+}
+
+void FAvalancheSequenceNavigatorModule::ShutdownModule()
+{
+	if (ISequencerModule* const SequencerModule = FModuleManager::LoadModulePtr<ISequencerModule>(TEXT("Sequencer")))
+	{
+		SequencerModule->UnregisterOnSequencerCreated(SequencerCreatedHandle);
+		SequencerCreatedHandle.Reset();
+	}
+
+	FAvaSequenceNavigatorCommands::Unregister();
+}
+
+void FAvalancheSequenceNavigatorModule::OnSequencerCreated(const TSharedRef<ISequencer> InSequencer)
+{
+	SequencerClosedHandle = InSequencer->OnCloseEvent().AddRaw(this, &FAvalancheSequenceNavigatorModule::OnSequencerClosed);
+
+	if (UObject* const PlaybackContextObject = InSequencer->GetPlaybackContext())
+	{
+		if (UWorld* const World = PlaybackContextObject->GetWorld())
+		{
+			if (UAvaSequencerSubsystem* const SequencerSubsystem = World->GetSubsystem<UAvaSequencerSubsystem>())
+			{
+				AvaSequencerCreatedHandle = SequencerSubsystem->OnSequencerCreated().AddRaw(this, &FAvalancheSequenceNavigatorModule::OnAvaSequencerCreated);
+			}
+		}
+	}
+}
+
+void FAvalancheSequenceNavigatorModule::OnSequencerClosed(const TSharedRef<ISequencer> InSequencer)
+{
+	if (!NavigationToolProvider.IsValid())
+	{
+		return;
+	}
+
+	const FName ToolId = FNavigationToolExtender::GetToolInstanceId(*InSequencer);
+
+	if (FNavigationToolExtender::UnregisterToolProvider(ToolId, NavigationToolProvider->GetIdentifier()))
+	{
+		NavigationToolProvider.Reset();
+		SequencerClosedHandle.Reset();
+	}
+}
+
+void FAvalancheSequenceNavigatorModule::OnAvaSequencerCreated(const TSharedRef<IAvaSequencer> InAvaSequencer)
+{
+	const TSharedPtr<ISequencer> Sequencer = InAvaSequencer->GetSequencer();
+	if (!Sequencer.IsValid())
+	{
+		return;
+	}
+
+	const FName ToolId = FNavigationToolExtender::GetToolInstanceId(*Sequencer);
+
+	const TSharedPtr<FNavigationToolProvider> FoundProvider = FNavigationToolExtender::FindToolProvider(ToolId
+		, FAvaNavigationToolProvider::Identifier);
+	if (!FoundProvider.IsValid())
+	{
+		NavigationToolProvider = MakeShared<FAvaNavigationToolProvider>(InAvaSequencer);
+		if (NavigationToolProvider->IsSequenceSupported(Sequencer->GetRootMovieSceneSequence()))
+		{
+			FNavigationToolExtender::RegisterToolProvider(Sequencer.ToSharedRef(), NavigationToolProvider.ToSharedRef());
+		}
+		else
+		{
+			NavigationToolProvider.Reset();
+		}
+	}
+}
+
+#undef LOCTEXT_NAMESPACE
+	
+IMPLEMENT_MODULE(FAvalancheSequenceNavigatorModule, AvalancheSequenceNavigator)
diff --git a/Engine/Plugins/Experimental/AvalancheSequenceNavigator/Source/AvalancheSequenceNavigator/Private/AvalancheSequenceNavigatorModule.h b/Engine/Plugins/Experimental/AvalancheSequenceNavigator/Source/AvalancheSequenceNavigator/Private/AvalancheSequenceNavigatorModule.h
new file mode 100644
--- /dev/null
+++ b/Engine/Plugins/Experimental/AvalancheSequenceNavigator/Source/AvalancheSequenceNavigator/Private/AvalancheSequenceNavigatorModule.h
@@ -0,0 +1,38 @@
+﻿// Copyright Epic Games, Inc. All Rights Reserved.
+
+#pragma once
+
+#include "Delegates/IDelegateInstance.h"
+#include "Modules/ModuleInterface.h"
+#include "Templates/SharedPointer.h"
+
+class IAvaSequencer;
+class ISequencer;
+
+namespace UE::AvaSequencer
+{
+	class FAvaNavigationToolProvider;
+}
+
+class FAvalancheSequenceNavigatorModule : public IModuleInterface
+{
+public:
+	//~ Begin IModuleInterface
+	virtual void StartupModule() override;
+	virtual void ShutdownModule() override;
+	//~ End IModuleInterface
+
+private:
+	void OnSequencerCreated(const TSharedRef<ISequencer> InSequencer);
+	void OnSequencerClosed(const TSharedRef<ISequencer> InSequencer);
+
+	void OnAvaSequencerCreated(const TSharedRef<IAvaSequencer> InAvaSequencer);
+
+	FDelegateHandle SequencerCreatedHandle;
+	FDelegateHandle SequencerClosedHandle;
+
+	FDelegateHandle AvaSequencerCreatedHandle;
+
+	/** The provider that supplies data and extends the Navigation Tool */
+	TSharedPtr<UE::AvaSequencer::FAvaNavigationToolProvider> NavigationToolProvider;
+};
diff --git a/Engine/Plugins/Experimental/AvalancheSequenceNavigator/Source/AvalancheSequenceNavigator/Private/Columns/AvaNavigationToolStatusColumn.cpp b/Engine/Plugins/Experimental/AvalancheSequenceNavigator/Source/AvalancheSequenceNavigator/Private/Columns/AvaNavigationToolStatusColumn.cpp
new file mode 100644
--- /dev/null
+++ b/Engine/Plugins/Experimental/AvalancheSequenceNavigator/Source/AvalancheSequenceNavigator/Private/Columns/AvaNavigationToolStatusColumn.cpp
@@ -0,0 +1,66 @@
+﻿// Copyright Epic Games, Inc. All Rights Reserved.
+
+#include "AvaNavigationToolStatusColumn.h"
+#include "Columns/SAvaNavigationToolStatus.h"
+#include "ISequencer.h"
+#include "Items/NavigationToolAvaSequence.h"
+#include "Styling/AppStyle.h"
+#include "INavigationTool.h"
+#include "INavigationToolView.h"
+
+#define LOCTEXT_NAMESPACE "AvaNavigationToolStatusColumn"
+
+namespace UE::SequenceNavigator
+{
+
+FText FAvaNavigationToolStatusColumn::GetColumnDisplayNameText() const
+{
+	return LOCTEXT("StatusColumn", "Status");
+}
+
+const FSlateBrush* FAvaNavigationToolStatusColumn::GetIconBrush() const
+{
+	return FAppStyle::GetBrush(TEXT("FoliageEditMode.BubbleBorder"));
+}
+
+SHeaderRow::FColumn::FArguments FAvaNavigationToolStatusColumn::ConstructHeaderRowColumn(const TSharedRef<INavigationToolView>& InToolView, const float InFillSize)
+{
+	const FName ColumnId = GetColumnId();
+
+	return SHeaderRow::Column(ColumnId)
+		.FillWidth(InFillSize)
+		.HAlignHeader(HAlign_Center)
+		.VAlignHeader(VAlign_Center)
+		.HAlignCell(HAlign_Fill)
+		.VAlignCell(VAlign_Center)
+		.DefaultLabel(GetColumnDisplayNameText())
+		.OnGetMenuContent(InToolView, &INavigationToolView::GetColumnMenuContent, ColumnId);
+}
+
+TSharedRef<SWidget> FAvaNavigationToolStatusColumn::ConstructRowWidget(const FNavigationToolItemRef& InItem
+	, const TSharedRef<INavigationToolView>& InView
+	, const TSharedRef<SNavigationToolTreeRow>& InRow)
+{
+	if (!InItem->IsA<FNavigationToolAvaSequence>())
+	{
+		return SNullWidget::NullWidget;
+	}
+
+	const TSharedPtr<INavigationTool> OwnerTool = InView->GetOwnerTool();
+	if (!OwnerTool.IsValid())
+	{
+		return SNullWidget::NullWidget;
+	}
+
+	const TSharedPtr<ISequencer> Sequencer = OwnerTool->GetSequencer();
+	if (!Sequencer.IsValid())
+	{
+		return SNullWidget::NullWidget;
+	}
+
+	return SNew(SAvaNavigationToolStatus, InItem, InView, InRow, Sequencer.ToSharedRef());
+}
+
+} // namespace UE::SequenceNavigator
+
+#undef LOCTEXT_NAMESPACE
diff --git a/Engine/Plugins/Experimental/AvalancheSequenceNavigator/Source/AvalancheSequenceNavigator/Private/Columns/AvaNavigationToolStatusColumn.h b/Engine/Plugins/Experimental/AvalancheSequenceNavigator/Source/AvalancheSequenceNavigator/Private/Columns/AvaNavigationToolStatusColumn.h
new file mode 100644
--- /dev/null
+++ b/Engine/Plugins/Experimental/AvalancheSequenceNavigator/Source/AvalancheSequenceNavigator/Private/Columns/AvaNavigationToolStatusColumn.h
@@ -0,0 +1,31 @@
+﻿// Copyright Epic Games, Inc. All Rights Reserved.
+
+#pragma once
+
+#include "Columns/NavigationToolColumn.h"
+
+namespace UE::SequenceNavigator
+{
+
+class FAvaNavigationToolStatusColumn : public FNavigationToolColumn
+{
+public:
+	UE_NAVIGATIONTOOL_INHERITS_WITH_SUPER(FAvaNavigationToolStatusColumn, FNavigationToolColumn);
+
+	static FName StaticColumnId() { return TEXT("Status"); }
+
+protected:
+	//~ Begin INavigationToolColumn
+	virtual FName GetColumnId() const override { return StaticColumnId(); }
+	virtual FText GetColumnDisplayNameText() const override;
+	virtual const FSlateBrush* GetIconBrush() const override;
+	virtual bool ShouldShowColumnByDefault() const override { return false; }
+	virtual float GetFillWidth() const override { return 10.f; }
+	virtual SHeaderRow::FColumn::FArguments ConstructHeaderRowColumn(const TSharedRef<INavigationToolView>& InToolView, const float InFillSize) override;
+	virtual TSharedRef<SWidget> ConstructRowWidget(const FNavigationToolItemRef& InItem
+		, const TSharedRef<INavigationToolView>& InView
+		, const TSharedRef<SNavigationToolTreeRow>& InRow) override;
+	//~ End INavigationToolColumn
+};
+
+} // namespace UE::SequenceNavigator
diff --git a/Engine/Plugins/Experimental/AvalancheSequenceNavigator/Source/AvalancheSequenceNavigator/Private/Columns/SAvaNavigationToolStatus.cpp b/Engine/Plugins/Experimental/AvalancheSequenceNavigator/Source/AvalancheSequenceNavigator/Private/Columns/SAvaNavigationToolStatus.cpp
new file mode 100644
--- /dev/null
+++ b/Engine/Plugins/Experimental/AvalancheSequenceNavigator/Source/AvalancheSequenceNavigator/Private/Columns/SAvaNavigationToolStatus.cpp
@@ -0,0 +1,146 @@
+﻿// Copyright Epic Games, Inc. All Rights Reserved.
+
+#include "SAvaNavigationToolStatus.h"
+#include "AvaSequence.h"
+#include "AvaSequencePlayer.h"
+#include "AvaSequencerUtils.h"
+#include "IAvaSceneInterface.h"
+#include "IAvaSequencerProvider.h"
+#include "INavigationToolView.h"
+#include "ISequencer.h"
+#include "Items/NavigationToolAvaSequence.h"
+#include "Widgets/Layout/SScrollBox.h"
+#include "Widgets/Notifications/SProgressBar.h"
+#include "Widgets/SNavigationToolTreeRow.h"
+
+#define LOCTEXT_NAMESPACE "SAvaNavigationToolStatus"
+
+namespace UE::SequenceNavigator
+{
+
+void SAvaNavigationToolStatus::Construct(const FArguments& InArgs
+	, const FNavigationToolItemPtr& InItem
+	, const TSharedRef<INavigationToolView>& InView
+	, const TSharedRef<SNavigationToolTreeRow>& InRowWidget
+	, const TSharedRef<ISequencer>& InSequencer)
+{
+	WeakItem = InItem;
+	WeakView = InView;
+	WeakRowWidget = InRowWidget;
+	WeakSequencer = InSequencer;
+
+	FNavigationToolAvaSequence* const AvaSequenceItem = InItem->CastTo<FNavigationToolAvaSequence>();
+	if (!AvaSequenceItem)
+	{
+		return;
+	}
+
+	ChildSlot
+	[
+		SNew(SBox)
+		.Padding(5.f, 2.f)
+		[
+			SNew(SOverlay)
+			+ SOverlay::Slot()
+			.Padding(1.f, 1.f)
+			[
+				SNew(SProgressBar)
+				.Percent(this, &SAvaNavigationToolStatus::GetProgressPercent)
+				.Visibility(EVisibility::Visible)
+			]
+			+ SOverlay::Slot()
+			.Padding(1.f, 0.f)
+			.VAlign(VAlign_Center)
+			[
+				SNew(STextBlock)
+				.Text(this, &SAvaNavigationToolStatus::GetProgressText)
+				.TextStyle(&FAppStyle::GetWidgetStyle<FTextBlockStyle>(TEXT("SmallText")))
+				.Justification(ETextJustify::Type::Center)
+			]
+		]
+	];
+}
+
+void SAvaNavigationToolStatus::Tick(const FGeometry& InAllottedGeometry, const double InCurrentTime, const float InDeltaTime)
+{
+	if (!WeakItem.IsValid() || !WeakSequencer.IsValid())
+	{
+		return;
+	}
+
+	const TSharedPtr<ISequencer> Sequencer = WeakSequencer.Pin();
+	if (!ensure(Sequencer.IsValid()))
+	{
+		return;
+	}
+
+	const FNavigationToolAvaSequence* const AvaSequenceItem = WeakItem.Pin()->CastTo<FNavigationToolAvaSequence>();
+	if (!ensure(AvaSequenceItem))
+	{
+		return;
+	}
+
+	UAvaSequence* const Sequence = AvaSequenceItem->GetAvaSequence();
+	if (!Sequence)
+	{
+		return;
+	}
+
+	IAvaSceneInterface* const SceneInterface = FAvaSequencerUtils::GetSceneInterface(Sequencer.ToSharedRef());
+	if (!ensure(SceneInterface))
+	{
+		return;
+	}
+
+	IAvaSequencePlaybackObject* const PlaybackObject = SceneInterface->GetPlaybackObject();
+	if (!PlaybackObject)
+	{
+		return;
+	}
+
+	UAvaSequencePlayer* const Player = PlaybackObject->GetSequencePlayer(Sequence);
+
+	bSequenceInProgress = Player != nullptr;
+	if (bSequenceInProgress)
+	{
+		const EMovieScenePlayerStatus::Type Status = Player->GetPlaybackStatus();
+		CurrentFrame = Player->GetCurrentTime().ConvertTo(Player->GetDisplayRate());
+		TotalFrames = Player->GetDuration().ConvertTo(Player->GetDisplayRate());
+
+		FFormatNamedArguments Args;
+		if (Status == EMovieScenePlayerStatus::Playing)
+		{
+			Args.Add("Status", LOCTEXT("SequenceStatus_Playing", "Playing"));
+		}
+		else
+		{
+			Args.Add("Status", LOCTEXT("SequenceStatus_Stopped" , "Stopped"));
+		}
+		Args.Add("Current", FText::AsNumber(CurrentFrame.GetFrame().Value));
+		Args.Add("Total", FText::AsNumber(TotalFrames.GetFrame().Value));
+
+		StatusText =  FText::Format(LOCTEXT("SequenceStatus_Text", "{Status} ({Current} / {Total})"), Args);
+
+		Progress = FMath::IsNearlyZero(TotalFrames.AsDecimal())
+			? 0.f : (CurrentFrame.AsDecimal() / TotalFrames.AsDecimal());
+	}
+	else
+	{
+		StatusText = LOCTEXT("SequenceStatus_Unknown", "Not Playing");
+		Progress = 0.f;
+	}
+}
+
+FText SAvaNavigationToolStatus::GetProgressText() const
+{
+	return StatusText;
+}
+
+TOptional<float> SAvaNavigationToolStatus::GetProgressPercent() const
+{
+	return Progress;
+}
+
+} // namespace UE::SequenceNavigator
+
+#undef LOCTEXT_NAMESPACE
diff --git a/Engine/Plugins/Experimental/AvalancheSequenceNavigator/Source/AvalancheSequenceNavigator/Private/Columns/SAvaNavigationToolStatus.h b/Engine/Plugins/Experimental/AvalancheSequenceNavigator/Source/AvalancheSequenceNavigator/Private/Columns/SAvaNavigationToolStatus.h
new file mode 100644
--- /dev/null
+++ b/Engine/Plugins/Experimental/AvalancheSequenceNavigator/Source/AvalancheSequenceNavigator/Private/Columns/SAvaNavigationToolStatus.h
@@ -0,0 +1,51 @@
+﻿// Copyright Epic Games, Inc. All Rights Reserved.
+
+#pragma once
+
+#include "Misc/FrameTime.h"
+#include "NavigationToolDefines.h"
+#include "Widgets/SCompoundWidget.h"
+
+class FText;
+class ISequencer;
+
+namespace UE::SequenceNavigator
+{
+
+class INavigationToolView;
+class SNavigationToolTreeRow;
+
+class SAvaNavigationToolStatus : public SCompoundWidget
+{
+public:
+	SLATE_BEGIN_ARGS(SAvaNavigationToolStatus) {}
+	SLATE_END_ARGS()
+
+	void Construct(const FArguments& InArgs
+		, const FNavigationToolItemPtr& InItem
+		, const TSharedRef<INavigationToolView>& InView
+		, const TSharedRef<SNavigationToolTreeRow>& InRowWidget
+		, const TSharedRef<ISequencer>& InSequencer);
+
+private:
+	//~ Begin SWidget
+	virtual void Tick(const FGeometry& InAllottedGeometry, const double InCurrentTime, const float InDeltaTime) override;
+	//~ End SWidget
+
+	FText GetProgressText() const;
+
+	TOptional<float> GetProgressPercent() const;
+
+	FNavigationToolItemWeakPtr WeakItem;
+	TWeakPtr<INavigationToolView> WeakView;
+	TWeakPtr<SNavigationToolTreeRow> WeakRowWidget;
+	TWeakPtr<ISequencer> WeakSequencer;
+
+	FText StatusText;
+	FFrameTime CurrentFrame;
+	FFrameTime TotalFrames;
+	float Progress = 0.f;
+	bool bSequenceInProgress = false;
+};
+
+} // namespace UE::SequenceNavigator
diff --git a/Engine/Plugins/Experimental/AvalancheSequenceNavigator/Source/AvalancheSequenceNavigator/Private/DragDrop/NavigationToolAvaSequenceDropHandler.cpp b/Engine/Plugins/Experimental/AvalancheSequenceNavigator/Source/AvalancheSequenceNavigator/Private/DragDrop/NavigationToolAvaSequenceDropHandler.cpp
new file mode 100644
--- /dev/null
+++ b/Engine/Plugins/Experimental/AvalancheSequenceNavigator/Source/AvalancheSequenceNavigator/Private/DragDrop/NavigationToolAvaSequenceDropHandler.cpp
@@ -0,0 +1,214 @@
+// Copyright Epic Games, Inc. All Rights Reserved.
+
+#include "NavigationToolAvaSequenceDropHandler.h"
+#include "AvaSequence.h"
+#include "IAvaSequenceProvider.h"
+#include "IAvaSequencer.h"
+#include "IAvaSequencerProvider.h"
+#include "INavigationTool.h"
+#include "ItemActions/NavigationToolAddItem.h"
+#include "ItemActions/NavigationToolRemoveItem.h"
+#include "Items/NavigationToolAvaSequence.h"
+#include "Misc/Optional.h"
+#include "ScopedTransaction.h"
+
+#define LOCTEXT_NAMESPACE "NavigationToolAvaSequenceDropHandler"
+
+namespace UE::SequenceNavigator
+{
+
+FNavigationToolAvaSequenceDropHandler::FNavigationToolAvaSequenceDropHandler(const TWeakPtr<IAvaSequencer>& InWeakAvaSequencer)
+	: WeakAvaSequencer(InWeakAvaSequencer)
+{
+	
+}
+
+bool FNavigationToolAvaSequenceDropHandler::IsDraggedItemSupported(const FNavigationToolItemPtr& InDraggedItem) const
+{
+	return InDraggedItem->IsA<FNavigationToolAvaSequence>();
+}
+
+TOptional<EItemDropZone> FNavigationToolAvaSequenceDropHandler::CanDrop(const EItemDropZone InDropZone, const FNavigationToolItemPtr& InTargetItem) const
+{
+	const bool bTargetIsRoot = (InTargetItem->GetItemId() == FNavigationToolItemId::RootId);
+
+	UAvaSequence* TargetAvaSequence = nullptr;
+
+	if (!bTargetIsRoot)
+	{
+		if (FNavigationToolAvaSequence* const TargetAvaSequenceItem = InTargetItem->CastTo<FNavigationToolAvaSequence>())
+		{
+			TargetAvaSequence = TargetAvaSequenceItem->GetAvaSequence();
+		}
+		if (!TargetAvaSequence)
+		{
+			return TOptional<EItemDropZone>();
+		}
+	}
+
+	for (const FNavigationToolItemPtr& Item : Items)
+	{
+		// Only allow moving sequences to other Motion Design sequences
+		const FNavigationToolAvaSequence* const AvaSequenceItemToMove = Item->CastTo<FNavigationToolAvaSequence>();
+		if (!AvaSequenceItemToMove)
+		{
+			return TOptional<EItemDropZone>();
+		}
+
+		// Only allow moving valid sequences and sequences that aren't the target
+		UAvaSequence* const AvaSequenceToMove = AvaSequenceItemToMove->GetAvaSequence();
+		if (!AvaSequenceToMove || AvaSequenceToMove == TargetAvaSequence)
+		{
+			return TOptional<EItemDropZone>();
+		}
+
+		if (!bTargetIsRoot && TargetAvaSequence)
+		{
+			if (TargetAvaSequence->GetChildren().Contains(AvaSequenceToMove))
+			{
+				return TOptional<EItemDropZone>();
+			}
+		}
+	}
+
+	switch (ActionType)
+	{
+		case ENavigationToolDragDropActionType::Move:
+			// Make sure the destination is not one of the items we're moving
+			if (!Items.Contains(InTargetItem))
+			{
+				return InDropZone;
+			}
+			break;
+
+		case ENavigationToolDragDropActionType::Copy:
+			return InDropZone;
+	}
+
+	return TOptional<EItemDropZone>();
+}
+
+bool FNavigationToolAvaSequenceDropHandler::Drop(const EItemDropZone InDropZone, const FNavigationToolItemPtr& InTargetItem)
+{
+	switch (ActionType)
+	{
+	case ENavigationToolDragDropActionType::Move:
+		MoveItems(InDropZone, InTargetItem);
+		break;
+
+	case ENavigationToolDragDropActionType::Copy:
+		DuplicateItems(Items, InTargetItem, InDropZone);
+		break;
+
+	default:
+		return false;
+	}
+
+	return true;
+}
+
+void FNavigationToolAvaSequenceDropHandler::MoveItems(const EItemDropZone InDropZone, const FNavigationToolItemPtr& InTargetItem)
+{
+	TArray<TSharedPtr<INavigationToolAction>> ItemActions;
+
+	FNavigationToolAddItemParams AddParams;
+	AddParams.RelativeItem = InTargetItem;
+	AddParams.RelativeDropZone = InDropZone;
+	AddParams.Flags = ENavigationToolAddItemFlags::Select | ENavigationToolAddItemFlags::Transact;
+	AddParams.SelectionFlags = ENavigationToolItemSelectionFlags::AppendToCurrentSelection;
+
+	const TSet<FNavigationToolItemPtr> DraggedItemSet(Items);
+
+	// Remove all Items whose Parent is in the Item Set 
+	Items.RemoveAll([&DraggedItemSet](const FNavigationToolItemPtr& InItem)
+		{
+			return !InItem.IsValid()
+				|| !InItem->GetParent().IsValid()
+				|| DraggedItemSet.Contains(InItem->GetParent());
+		});
+
+	// Reverse order for onto since Item->AddChild(...) adds it to Index 0, so last item would be at the top, which is reversed
+	if (InDropZone == EItemDropZone::OntoItem)
+	{
+		Algo::Reverse(Items);
+	}
+
+	for (const FNavigationToolItemPtr& Item : Items)
+	{
+		// Only allow moving sequences to other Motion Design sequences
+		if (FNavigationToolAvaSequence* const AvaSequenceItemToMove = Item->CastTo<FNavigationToolAvaSequence>())
+		{
+			AddParams.Item = SharedThis(AvaSequenceItemToMove);
+			ItemActions.Add(MakeShared<FNavigationToolAddItem>(AddParams));
+		}
+	}
+
+	InTargetItem->GetOwnerTool().EnqueueItemActions(MoveTemp(ItemActions));
+}
+
+void FNavigationToolAvaSequenceDropHandler::DuplicateItems(const TArray<FNavigationToolItemPtr>& InItems
+	, const FNavigationToolItemPtr& InRelativeItem
+	, const TOptional<EItemDropZone>& InRelativeDropZone)
+{
+	const TSharedPtr<IAvaSequencer> AvaSequencer = WeakAvaSequencer.Pin();
+	if (!AvaSequencer.IsValid())
+	{
+		return;
+	}
+
+	const IAvaSequencerProvider& SequencerProvider = AvaSequencer->GetProvider();
+
+	IAvaSequenceProvider* const SequenceProvider = SequencerProvider.GetSequenceProvider();
+	if (!SequenceProvider)
+	{
+		return;
+	}
+
+	UObject* const Outer = SequenceProvider->ToUObject();
+	if (!Outer)
+	{
+		return;
+	}
+
+	// Gather list of sequences to duplicate
+	TSet<UAvaSequence*> SequencesToDuplicate;
+
+	for (const FNavigationToolItemPtr& Item : InItems)
+	{
+		const FNavigationToolAvaSequence* const AvaSequenceItem = Item->CastTo<FNavigationToolAvaSequence>();
+		if (!AvaSequenceItem)
+		{
+			continue;
+		}
+
+		UAvaSequence* const TemplateSequence = AvaSequenceItem->GetAvaSequence();
+		if (!TemplateSequence)
+		{
+			continue;
+		}
+
+		SequencesToDuplicate.Add(TemplateSequence);
+	}
+
+	if (SequencesToDuplicate.IsEmpty())
+	{
+		return;
+	}
+
+	// Duplicate sequence objects
+	const FScopedTransaction Transaction(LOCTEXT("DuplicateSequencesTransaction", "Duplicate Sequence(s)"));
+
+	Outer->Modify();
+
+	for (const UAvaSequence* const TemplateSequence : SequencesToDuplicate)
+	{
+		if (UAvaSequence* const DupedSequence = DuplicateObject<UAvaSequence>(TemplateSequence, Outer))
+		{
+			SequenceProvider->AddSequence(DupedSequence);
+		}
+	}
+}
+
+} // namespace UE::SequenceNavigator
+
+#undef LOCTEXT_NAMESPACE
diff --git a/Engine/Plugins/Experimental/AvalancheSequenceNavigator/Source/AvalancheSequenceNavigator/Private/DragDrop/NavigationToolAvaSequenceDropHandler.h b/Engine/Plugins/Experimental/AvalancheSequenceNavigator/Source/AvalancheSequenceNavigator/Private/DragDrop/NavigationToolAvaSequenceDropHandler.h
new file mode 100644
--- /dev/null
+++ b/Engine/Plugins/Experimental/AvalancheSequenceNavigator/Source/AvalancheSequenceNavigator/Private/DragDrop/NavigationToolAvaSequenceDropHandler.h
@@ -0,0 +1,38 @@
+﻿// Copyright Epic Games, Inc. All Rights Reserved.
+
+#pragma once
+
+#include "Templates/SharedPointer.h"
+#include "DragDropOps/Handlers/NavigationToolItemDropHandler.h"
+
+class IAvaSequencer;
+enum class EItemDropZone;
+template<typename OptionalType> struct TOptional;
+
+namespace UE::SequenceNavigator
+{
+
+class FNavigationToolAvaSequenceDropHandler : public SequenceNavigator::FNavigationToolItemDropHandler
+{
+public:
+	UE_NAVIGATIONTOOL_INHERITS(FNavigationToolAvaSequenceDropHandler, SequenceNavigator::FNavigationToolItemDropHandler);
+
+	FNavigationToolAvaSequenceDropHandler(const TWeakPtr<IAvaSequencer>& InWeakAvaSequencer);
+
+protected:
+	//~ Begin FNavigationToolItemDropHandler
+	virtual bool IsDraggedItemSupported(const SequenceNavigator::FNavigationToolItemPtr& InDraggedItem) const override;
+	virtual TOptional<EItemDropZone> CanDrop(const EItemDropZone InDropZone, const SequenceNavigator::FNavigationToolItemPtr& InTargetItem) const override;
+	virtual bool Drop(const EItemDropZone InDropZone, const SequenceNavigator::FNavigationToolItemPtr& InTargetItem) override;
+	//~ End FNavigationToolItemDropHandler
+
+	void MoveItems(const EItemDropZone InDropZone, const SequenceNavigator::FNavigationToolItemPtr& InTargetItem);
+
+	void DuplicateItems(const TArray<SequenceNavigator::FNavigationToolItemPtr>& InItems
+		, const SequenceNavigator::FNavigationToolItemPtr& InRelativeItem
+		, const TOptional<EItemDropZone>& InRelativeDropZone);
+
+	TWeakPtr<IAvaSequencer> WeakAvaSequencer;
+};
+
+} // namespace UE::SequenceNavigator
diff --git a/Engine/Plugins/Experimental/AvalancheSequenceNavigator/Source/AvalancheSequenceNavigator/Private/Items/NavigationToolAvaSequence.cpp b/Engine/Plugins/Experimental/AvalancheSequenceNavigator/Source/AvalancheSequenceNavigator/Private/Items/NavigationToolAvaSequence.cpp
new file mode 100644
--- /dev/null
+++ b/Engine/Plugins/Experimental/AvalancheSequenceNavigator/Source/AvalancheSequenceNavigator/Private/Items/NavigationToolAvaSequence.cpp
@@ -0,0 +1,316 @@
+﻿// Copyright Epic Games, Inc. All Rights Reserved.
+
+#include "NavigationToolAvaSequence.h"
+#include "AvaEditorCoreStyle.h"
+#include "AvaSequence.h"
+#include "AvaSequencerUtils.h"
+#include "IAvaSequencer.h"
+#include "INavigationTool.h"
+#include "ISequencer.h"
+#include "Items/NavigationToolItemParameters.h"
+#include "Items/NavigationToolSequence.h"
+#include "ScopedTransaction.h"
+#include "Styling/StyleColors.h"
+
+#define LOCTEXT_NAMESPACE "NavigationToolAvaSequence"
+
+namespace UE::SequenceNavigator
+{
+
+using namespace UE::Sequencer;
+
+FNavigationToolAvaSequence::FNavigationToolAvaSequence(INavigationTool& InTool
+	, const FNavigationToolItemPtr& InParentItem
+	, UAvaSequence* const InAvaSequence)
+	: FNavigationToolSequence(InTool
+	, InParentItem
+	, InAvaSequence
+	, nullptr
+	, 0)
+{
+}
+
+bool FNavigationToolAvaSequence::AddChild(const FNavigationToolAddItemParams& InAddItemParams)
+{
+	if (!CanAddChild(InAddItemParams.Item))
+	{
+		return false;
+	}
+
+	const FNavigationToolAvaSequence* const AvaSequenceItemToAdd = InAddItemParams.Item->CastTo<FNavigationToolAvaSequence>();
+	if (!AvaSequenceItemToAdd)
+	{
+		return false;
+	}
+
+	AddChildChecked(InAddItemParams);
+
+	UAvaSequence* const ParentAvaSequence = GetAvaSequence();
+	UAvaSequence* const AvaSequenceToAdd = AvaSequenceItemToAdd->GetAvaSequence();
+
+	const TArray<TWeakObjectPtr<UAvaSequence>> TargetChildren = ParentAvaSequence->GetChildren();
+	if (TargetChildren.Contains(AvaSequenceToAdd))
+	{
+		return false;
+	}
+
+	// Remove the sequence from any parent it's a child of
+	if (UAvaSequence* const AvaSequenceToMoveParent = AvaSequenceToAdd->GetParent())
+	{
+		AvaSequenceToMoveParent->RemoveChild(AvaSequenceToAdd);
+	}
+
+	// Process moving sequence if it's parent is the target
+	if (ParentAvaSequence == AvaSequenceToAdd->GetParent())
+	{
+		ParentAvaSequence->RemoveChild(AvaSequenceToAdd);
+
+		if (UAvaSequence* const TargetParent = ParentAvaSequence->GetParent())
+		{
+			TargetParent->AddChild(AvaSequenceToAdd);
+		}
+
+		return true;
+	}
+
+	// Check if the Target Item is a Child of Sequence Item (to prevent cycle)
+	UAvaSequence* CurrentItem = ParentAvaSequence;
+	while (CurrentItem)
+	{
+		if (CurrentItem->GetParent() == AvaSequenceToAdd)
+		{
+			AvaSequenceToAdd->RemoveChild(CurrentItem);
+			CurrentItem->SetParent(AvaSequenceToAdd->GetParent());
+
+			break;
+		}
+
+		CurrentItem = CurrentItem->GetParent();
+	}
+
+	ParentAvaSequence->AddChild(AvaSequenceToAdd);
+
+	return true;
+}
+
+bool FNavigationToolAvaSequence::RemoveChild(const FNavigationToolRemoveItemParams& InRemoveItemParams)
+{
+	// Check that we're removing a child item that is directly under us, and not a grandchild
+	if (!InRemoveItemParams.Item.IsValid() || !Children.Contains(InRemoveItemParams.Item))
+	{
+		return false;
+	}
+
+	if (const FNavigationToolAvaSequence* const AvaSequenceItemToRemove = InRemoveItemParams.Item->CastTo<FNavigationToolAvaSequence>())
+	{
+		if (UAvaSequence* const AvaSequenceToRemove = AvaSequenceItemToRemove->GetAvaSequence())
+		{
+			// Remove the sequence from any parent it's a child of
+			if (UAvaSequence* const AvaSequenceToRemoveParent = AvaSequenceToRemove->GetParent())
+			{
+				AvaSequenceToRemoveParent->RemoveChild(AvaSequenceToRemove);
+			}
+		}
+	}
+
+	return RemoveChildChecked(InRemoveItemParams.Item);
+}
+
+void FNavigationToolAvaSequence::FindChildren(TArray<FNavigationToolItemPtr>& OutChildren, bool bRecursive)
+{
+	FNavigationToolSequence::FindChildren(OutChildren, bRecursive);
+
+	UAvaSequence* const AvaSequence = GetAvaSequence();
+	if (!AvaSequence)
+	{
+		return;
+	}
+
+	const TSharedPtr<ISequencer> Sequencer = Tool.GetSequencer();
+	if (!Sequencer.IsValid())
+	{
+		return;
+	}
+
+	const TSharedPtr<FNavigationToolProvider> Provider = GetProvider();
+	if (!Provider.IsValid())
+	{
+		return;
+	}
+
+	TArray<TWeakObjectPtr<UAvaSequence>> ChildSequences = AvaSequence->GetChildren();
+	ChildSequences.RemoveAll([](const TWeakObjectPtr<UAvaSequence>& InChild)
+		{
+			return !InChild.IsValid();
+		});
+	ChildSequences.Sort([](const TWeakObjectPtr<UAvaSequence>& InA, const TWeakObjectPtr<UAvaSequence>& InB)
+		{
+			return InA->GetDisplayName().CompareTo(InB->GetDisplayName()) > 0;
+		});
+
+	const TSharedRef<FNavigationToolAvaSequence> SharedThisRef = SharedThis(this);
+
+	for (const TWeakObjectPtr<UAvaSequence>& SequenceWeak : ChildSequences)
+	{
+		const FNavigationToolItemPtr NewItem = Tool.FindOrAdd<FNavigationToolAvaSequence>(Provider
+			, SharedThisRef, SequenceWeak.Get());
+		OutChildren.Add(NewItem);
+		if (bRecursive)
+		{
+			NewItem->FindChildren(OutChildren, bRecursive);
+		}
+	}
+}
+
+void FNavigationToolAvaSequence::GetItemProxies(TArray<TSharedPtr<FNavigationToolItemProxy>>& OutItemProxies)
+{
+	FNavigationToolSequence::GetItemProxies(OutItemProxies);
+}
+
+bool FNavigationToolAvaSequence::CanRename() const
+{
+	return GetAvaSequence() != nullptr;
+}
+
+bool FNavigationToolAvaSequence::Rename(const FString& InName)
+{
+	UAvaSequence* const UnderlyingSequence = GetAvaSequence();
+	if (!UnderlyingSequence)
+	{
+		return false;
+	}
+
+	bool bRenamed = false;
+
+	if (!InName.Equals(UnderlyingSequence->GetLabel().ToString(), ESearchCase::CaseSensitive))
+	{
+		const FScopedTransaction Transaction(LOCTEXT("AvaSequenceRename", "Rename Sequence"));
+
+		UnderlyingSequence->Modify();
+		UnderlyingSequence->SetLabel(FName(InName));
+
+		Tool.NotifyToolItemRenamed(SharedThis(this));
+
+		bRenamed = true;
+	}
+
+	return bRenamed;
+}
+
+bool FNavigationToolAvaSequence::CanDelete() const
+{
+	return GetAvaSequence() != nullptr;
+}
+
+bool FNavigationToolAvaSequence::Delete()
+{
+	const TSharedPtr<ISequencer> Sequencer = Tool.GetSequencer();
+	if (!Sequencer.IsValid())
+	{
+		return false;
+	}
+
+	UAvaSequence* const AvaSequence = GetAvaSequence();
+	if (!AvaSequence)
+	{
+		return false;
+	}
+
+	const TSharedPtr<IAvaSequencer> AvaSequencer = FAvaSequencerUtils::GetAvaSequencer(Sequencer.ToSharedRef());
+	if (!AvaSequencer.IsValid())
+	{
+		return false;
+	}
+
+	AvaSequencer->DeleteSequences({ AvaSequence });
+
+	return true;
+}
+
+FText FNavigationToolAvaSequence::GetDisplayName() const
+{
+	if (const UAvaSequence* const AvaSequence = GetAvaSequence())
+	{
+		return AvaSequence->GetDisplayName();
+	}
+	return FText::GetEmpty();
+}
+
+FSlateIcon FNavigationToolAvaSequence::GetIcon() const
+{
+	return FSlateIcon(FAvaEditorCoreStyle::Get().GetStyleSetName(), TEXT("Icons.MotionDesign"));
+}
+
+FSlateColor FNavigationToolAvaSequence::GetIconColor() const
+{
+	const TSharedPtr<ISequencer> Sequencer = Tool.GetSequencer();
+	if (!Sequencer.IsValid())
+	{
+		return FStyleColors::Foreground;
+	}
+
+	const UMovieSceneSequence* const RootSequence = Sequencer->GetRootMovieSceneSequence();
+	if (!RootSequence)
+	{
+		return FStyleColors::Foreground;
+	}
+
+	if (GetAvaSequence() == RootSequence)
+	{
+		return FStyleColors::AccentGreen;
+	}
+
+	return FStyleColors::Foreground;
+}
+
+void FNavigationToolAvaSequence::OnSelect()
+{
+	Super::OnSelect();
+}
+
+void FNavigationToolAvaSequence::OnDoubleClick()
+{
+	UAvaSequence* const AvaSequence = GetAvaSequence();
+	if (!AvaSequence)
+	{
+		return;
+	}
+
+	const TSharedPtr<ISequencer> Sequencer = Tool.GetSequencer();
+	if (!Sequencer.IsValid())
+	{
+		return;
+	}
+
+	const TSharedPtr<IAvaSequencer> AvaSequencer = FAvaSequencerUtils::GetAvaSequencer(Sequencer.ToSharedRef());
+	if (!AvaSequencer.IsValid())
+	{
+		return;
+	}
+
+	AvaSequencer->SetViewedSequence(AvaSequence);
+}
+
+UAvaSequence* FNavigationToolAvaSequence::GetAvaSequence() const
+{
+	return Cast<UAvaSequence>(GetSequence());
+}
+
+FNavigationToolItemId FNavigationToolAvaSequence::CalculateItemId() const
+{
+	return FNavigationToolItemId(GetParent(), GetAvaSequence());
+}
+
+TOptional<FColor> FNavigationToolAvaSequence::GetColor() const
+{
+	return Tool.FindItemColor(SharedThis(const_cast<FNavigationToolAvaSequence*>(this)));
+}
+
+void FNavigationToolAvaSequence::SetColor(const TOptional<FColor>& InColor)
+{
+	Tool.SetItemColor(SharedThis(this), InColor.Get(FColor()));
+}
+
+} // namespace UE::SequenceNavigator
+
+#undef LOCTEXT_NAMESPACE
diff --git a/Engine/Plugins/Experimental/AvalancheSequenceNavigator/Source/AvalancheSequenceNavigator/Private/Items/NavigationToolAvaSequence.h b/Engine/Plugins/Experimental/AvalancheSequenceNavigator/Source/AvalancheSequenceNavigator/Private/Items/NavigationToolAvaSequence.h
new file mode 100644
--- /dev/null
+++ b/Engine/Plugins/Experimental/AvalancheSequenceNavigator/Source/AvalancheSequenceNavigator/Private/Items/NavigationToolAvaSequence.h
@@ -0,0 +1,62 @@
+﻿// Copyright Epic Games, Inc. All Rights Reserved.
+
+#pragma once
+
+#include "Items/NavigationToolSequence.h"
+
+class UAvaSequence;
+
+namespace UE::SequenceNavigator
+{
+
+/**
+ * Item in Navigation Tool representing a Motion Design Sequence.
+ */
+class FNavigationToolAvaSequence : public FNavigationToolSequence
+{
+public:
+	UE_NAVIGATIONTOOL_INHERITS_WITH_SUPER(FNavigationToolAvaSequence, FNavigationToolSequence);
+
+	FNavigationToolAvaSequence(INavigationTool& InTool
+		, const FNavigationToolItemPtr& InParentItem
+		, UAvaSequence* const InAvaSequence);
+
+	//~ Begin INavigationToolItem
+
+	virtual bool AddChild(const FNavigationToolAddItemParams& InAddItemParams) override;
+	virtual bool RemoveChild(const FNavigationToolRemoveItemParams& InRemoveItemParams) override;
+	virtual void FindChildren(TArray<FNavigationToolItemPtr>& OutChildren, bool bRecursive) override;
+	virtual void GetItemProxies(TArray<TSharedPtr<FNavigationToolItemProxy>>& OutItemProxies) override;
+
+	virtual bool ShouldSort() const override { return true; }
+	virtual bool CanBeTopLevel() const override { return true; }
+
+	virtual bool CanRename() const override;
+	virtual bool Rename(const FString& InName) override;
+
+	virtual bool CanDelete() const override;
+	virtual bool Delete() override;
+
+	virtual FText GetDisplayName() const override;
+	virtual FSlateIcon GetIcon() const override;
+	virtual FSlateColor GetIconColor() const override;
+
+	virtual void OnSelect() override;
+	virtual void OnDoubleClick() override;
+
+	//~ End INavigationToolItem
+
+	//~ Begin IColorExtension
+	virtual TOptional<FColor> GetColor() const override;
+	virtual void SetColor(const TOptional<FColor>& InColor) override;
+	//~ End IColorExtension
+
+	UAvaSequence* GetAvaSequence() const;
+
+protected:
+	//~Begin FNavigationToolItem
+	virtual FNavigationToolItemId CalculateItemId() const override;
+	//~End FNavigationToolItem
+};
+
+} // namespace UE::SequenceNavigator
diff --git a/Engine/Plugins/VirtualProduction/Avalanche/Source/AvalancheEditor/Private/Sequencer/AvaSequencerExtension.h b/Engine/Plugins/VirtualProduction/Avalanche/Source/AvalancheEditor/Private/Sequencer/AvaSequencerExtension.h
--- a/Engine/Plugins/VirtualProduction/Avalanche/Source/AvalancheEditor/Private/Sequencer/AvaSequencerExtension.h
+++ b/Engine/Plugins/VirtualProduction/Avalanche/Source/AvalancheEditor/Private/Sequencer/AvaSequencerExtension.h
@@ -24,8 +24,6 @@ public:
 
 	explicit FAvaSequencerExtension(FName InSequencerTabId, bool bInSupportsDrawerWidget);
 
-	TSharedPtr<IAvaSequencer> GetAvaSequencer() const { return AvaSequencer; }
-
 	TSharedPtr<ISequencer> GetSequencer() const;
 
 	void OnObjectRenamed(UObject* InRenamedObject, const FText& InDisplayNameText);
@@ -35,6 +33,7 @@ public:
 	virtual FAvaSequencerArgs MakeSequencerArgs() const;
 
 	//~ Begin IAvaEditorExtension
+	virtual TSharedPtr<IAvaSequencer> GetAvaSequencer() const override { return AvaSequencer; }
 	virtual void Activate() override;
 	virtual void Deactivate() override;
 	virtual void Cleanup() override;
diff --git a/Engine/Plugins/VirtualProduction/Avalanche/Source/AvalancheSequencer/AvalancheSequencer.Build.cs b/Engine/Plugins/VirtualProduction/Avalanche/Source/AvalancheSequencer/AvalancheSequencer.Build.cs
--- a/Engine/Plugins/VirtualProduction/Avalanche/Source/AvalancheSequencer/AvalancheSequencer.Build.cs
+++ b/Engine/Plugins/VirtualProduction/Avalanche/Source/AvalancheSequencer/AvalancheSequencer.Build.cs
@@ -50,7 +50,6 @@ public class AvalancheSequencer : ModuleRules
 				"Projects",
 				"SceneOutliner",
 				"SequencerCore",
-				"SequenceNavigator",
 				"Settings",
 				"Slate",
 				"StateTreeEditorModule",
diff --git a/Engine/Plugins/VirtualProduction/Avalanche/Source/AvalancheSequencer/Private/AvaSequencer.cpp b/Engine/Plugins/VirtualProduction/Avalanche/Source/AvalancheSequencer/Private/AvaSequencer.cpp
--- a/Engine/Plugins/VirtualProduction/Avalanche/Source/AvalancheSequencer/Private/AvaSequencer.cpp
+++ b/Engine/Plugins/VirtualProduction/Avalanche/Source/AvalancheSequencer/Private/AvaSequencer.cpp
@@ -7,6 +7,7 @@
 #include "AvaSequencePlaybackObject.h"
 #include "AvaSequencePlayer.h"
 #include "AvaSequencerArgs.h"
+#include "AvaSequencerSubsystem.h"
 #include "AvaSequencerUtils.h"
 #include "Clipboard/AvaSequenceExporter.h"
 #include "Clipboard/AvaSequenceImporter.h"
@@ -27,10 +28,8 @@
 #include "Framework/Commands/GenericCommands.h"
 #include "Framework/Commands/UIAction.h"
 #include "Framework/MultiBox/MultiBoxBuilder.h"
-#include "Framework/Notifications/NotificationManager.h"
 #include "IAvaSequenceProvider.h"
 #include "IContentBrowserSingleton.h"
-#include "INavigationTool.h"
 #include "ISequencer.h"
 #include "ISequencerModule.h"
 #include "ISequencerTrackEditor.h"
@@ -41,9 +40,6 @@
 #include "Misc/MessageDialog.h"
 #include "Misc/NotifyHook.h"
 #include "MovieScene.h"
-#include "NavigationToolDefines.h"
-#include "NavigationTool/AvaNavigationToolProvider.h"
-#include "NavigationTool/NavigationToolAvaSequence.h"
 #include "Playback/AvaSequencerCleanView.h"
 #include "Playback/AvaSequencerController.h"
 #include "PropertyHandle.h"
@@ -61,7 +57,6 @@
 #include "Sidebar/SidebarDrawerConfig.h"
 #include "StaggerTool/AvaStaggerTool.h"
 #include "Toolkits/AssetEditorToolkit.h"
-#include "Widgets/Notifications/SNotificationList.h"
 #include "Widgets/Views/SHeaderRow.h"
 
 DEFINE_LOG_CATEGORY_STATIC(LogAvaSequencer, Log, All);
@@ -297,8 +292,6 @@ FAvaSequencer::~FAvaSequencer()
 			SidebarExtender.Reset();
 		}
 	}
-
-	UnregisterNavigationTool();
 }
 
 void FAvaSequencer::BindCommands()
@@ -958,9 +951,9 @@ void FAvaSequencer::ApplyDefaultPresetToSelection(FName InPresetName)
 		return;
 	}
 
-	TConstArrayView<FAvaSequencePreset> DefaultSequencePresets = SequencerSettings->GetDefaultSequencePresets();
+	const TConstArrayView<FAvaSequencePreset> DefaultSequencePresets = SequencerSettings->GetDefaultSequencePresets();
 
-	int32 PresetIndex = DefaultSequencePresets.Find(FAvaSequencePreset(InPresetName));
+	const int32 PresetIndex = DefaultSequencePresets.Find(FAvaSequencePreset(InPresetName));
 	if (PresetIndex == INDEX_NONE)
 	{
 		return;
@@ -988,33 +981,19 @@ void FAvaSequencer::ApplyCustomPresetToSelection(FName InPresetName)
 
 void FAvaSequencer::ApplyPresetToSelection(const FAvaSequencePreset& InPreset)
 {
-	using namespace UE::SequenceNavigator;
-
-	const TSharedPtr<ISequencer> Sequencer = SequencerWeak.Pin();
-	if (!Sequencer.IsValid())
-	{
-		return;
-	}
-
-	const TSharedPtr<INavigationTool> Tool = FNavigationToolExtender::FindNavigationTool(Sequencer.ToSharedRef());
-	if (!Tool.IsValid())
-	{
-		return;
-	}
-
-	const TArray<FNavigationToolItemPtr> SelectedItems = Tool->GetSelectedItems();
-	if (SelectedItems.IsEmpty())
+	const TArray<UAvaSequence*> SelectedSequences = GetSelectedSequences();
+	if (SelectedSequences.IsEmpty())
 	{
 		return;
 	}
 
 	FScopedTransaction Transaction(LOCTEXT("ApplySequencePreset", "Apply Sequence Preset"));
 
-	for (const FNavigationToolItemPtr& Item : SelectedItems)
+	for (UAvaSequence* const AvaSequence : SelectedSequences)
 	{
-		if (FNavigationToolAvaSequence* const AvaSequenceItem = Item->CastTo<FNavigationToolAvaSequence>())
+		if (AvaSequence)
 		{
-			InPreset.ApplyPreset(AvaSequenceItem->GetAvaSequence());
+			InPreset.ApplyPreset(AvaSequence);
 		}
 	}
 }
@@ -1333,139 +1312,6 @@ UAvaSequence* FAvaSequencer::AddSequence(UAvaSequence* const InParentSequence)
 	return NewSequence;
 }
 
-bool FAvaSequencer::CanDuplicateSequences() const
-{
-	return GetProvider().CanEditOrPlaySequences();
-}
-
-void FAvaSequencer::DuplicateSequences()
-{
-	const TArray<UAvaSequence*> SelectedSequences = GetSelectedSequences();
-	if (SelectedSequences.IsEmpty())
-	{
-		return;
-	}
-
-	IAvaSequenceProvider* const SequenceProvider = Provider.GetSequenceProvider();
-	if (!SequenceProvider)
-	{
-		return;
-	}
-
-	UObject* const Outer = SequenceProvider->ToUObject();
-	if (!Outer)
-	{
-		return;
-	}
-
-	const FScopedTransaction Transaction(LOCTEXT("DuplicateSequenceTransaction", "Duplicate Sequence"));
-
-	Outer->Modify();
-
-	for (UAvaSequence* const TemplateSequence : SelectedSequences)
-	{
-		if (TemplateSequence)
-		{
-			UAvaSequence* const DupedSequence = DuplicateObject<UAvaSequence>(TemplateSequence, Outer);
-			SequenceProvider->AddSequence(DupedSequence);
-		}
-	}
-}
-
-bool FAvaSequencer::CanExportSequences() const
-{
-	return GetSelectedSequenceItems().Num() > 0;
-}
-
-void FAvaSequencer::ExportSequences()
-{
-	const TArray<UAvaSequence*> SelectedSequences = GetSelectedSequences();
-
-	TArray<TWeakObjectPtr<>> WeakSelectedObjects;
-	Algo::Transform(SelectedSequences, WeakSelectedObjects, [](UAvaSequence* const InSequence)
-		{
-			return InSequence;
-		});
-
-	Provider.ExportSequences(SelectedSequences);
-
-	FNotificationInfo Info(LOCTEXT("ExportSuccess", "Sequence Exported Successfully!"));
-	Info.bUseLargeFont = false;
-	Info.bUseSuccessFailIcons = false;
-	Info.ExpireDuration = 5.0f;
-	Info.HyperlinkText = LOCTEXT("ShowNewAssetsInContentBrowser", "Show in content browser");
-	Info.Hyperlink = FSimpleDelegate::CreateLambda([WeakSelectedObjects]()
-		{
-			TArray<UObject*> SyncObjects;
-
-			for (const TWeakObjectPtr<> WeakObject : WeakSelectedObjects)
-			{
-				if (UObject* const SyncObject = WeakObject.Get())
-				{
-					SyncObjects.Add(SyncObject);
-				}
-			}
-
-			FContentBrowserModule& ContentBrowserModule = FModuleManager::LoadModuleChecked<FContentBrowserModule>(TEXT("ContentBrowser"));
-			ContentBrowserModule.Get().SyncBrowserToAssets(SyncObjects);
-		});
-
-	FSlateNotificationManager::Get().AddNotification(Info);
-}
-
-bool FAvaSequencer::CanSpawnPlayers() const
-{
-	return GetSelectedSequenceItems().Num() > 0;
-}
-
-void FAvaSequencer::SpawnPlayers()
-{
-	if (!GEditor)
-	{
-		return;
-	}
-
-	UWorld* World = Provider.GetPlaybackContext()->GetWorld();
-	if (!World)
-	{
-		return;
-	}
-
-	UActorFactory* ActorFactory = GEditor->FindActorFactoryForActorClass(AAvaSequenceActor::StaticClass());
-	if (!ensure(ActorFactory))
-	{
-		return;
-	}
-
-	TArray<UAvaSequence*> Sequences = GetSelectedSequences();
-
-	FScopedTransaction Transaction(LOCTEXT("SpawnSequencePlayers", "Spawn Sequence Players"));
-
-	for (UAvaSequence* Sequence : Sequences)
-	{
-		check(Sequence);
-		GEditor->UseActorFactory(ActorFactory, FAssetData(Sequence), &FTransform::Identity);
-	}
-
-	const FText NotifyText = FText::Format(LOCTEXT("SpawnPlayerSuccess", "{0} Sequence Players Spawned Successfully!")
-		, FText::FromString(FString::FromInt(Sequences.Num())));
-
-	FNotificationInfo Info(NotifyText);
-	Info.bUseLargeFont = false;
-	Info.bUseSuccessFailIcons = false;
-	FSlateNotificationManager::Get().AddNotification(Info);
-}
-
-bool FAvaSequencer::CanDeleteSequence() const
-{
-	return GetProvider().CanEditOrPlaySequences();
-}
-
-void FAvaSequencer::DeleteSequence()
-{
-	DeleteSequences(TSet<UAvaSequence*> { GetSelectedSequences() });
-}
-
 void FAvaSequencer::DeleteSequences(const TSet<UAvaSequence*>& InSequences)
 {
 	if (InSequences.IsEmpty())
@@ -1520,83 +1366,6 @@ void FAvaSequencer::DeleteSequences(const TSet<UAvaSequence*>& InSequences)
 	}
 }
 
-bool FAvaSequencer::CanRelabelSequence() const
-{
-	const bool bCanEditSequences = GetProvider().CanEditOrPlaySequences();
-	return bCanEditSequences && SequenceTreeView->GetNumItemsSelected() == 1;
-}
-
-void FAvaSequencer::RelabelSequence()
-{
-	TArray<FAvaSequenceItemPtr> SelectedItems = SequenceTreeView->GetSelectedItems();
-	check(SelectedItems.Num() == 1);
-
-	const FAvaSequenceItemPtr& SelectedItem = SelectedItems[0];
-	SelectedItem->RequestRelabel();
-}
-
-bool FAvaSequencer::CanPlaySequences() const
-{
-	const bool bCanEditSequences = GetProvider().CanEditOrPlaySequences();
-	return bCanEditSequences && GetSelectedSequenceItems().Num() > 0;
-}
-
-void FAvaSequencer::PlaySequences()
-{
-	const TArray<UAvaSequence*> SelectedSequences = GetSelectedSequences();
-
-	FAvaSequencePlayParams PlaySettings;
-	PlaySettings.AdvancedSettings.bRestoreState = true;
-
-	for (UAvaSequence* const Sequence : SelectedSequences)
-	{
-		IAvaSequencePlaybackObject* const PlaybackObject = Provider.GetPlaybackObject();
-		if (Sequence && PlaybackObject)
-		{
-			PlaybackObject->PlaySequence(Sequence, PlaySettings);
-		}
-	}
-}
-
-bool FAvaSequencer::CanContinueSequences() const
-{
-	const bool bCanEditSequences = GetProvider().CanEditOrPlaySequences();
-	return bCanEditSequences && GetSelectedSequenceItems().Num() > 0;
-}
-
-void FAvaSequencer::ContinueSequences()
-{
-	const TArray<UAvaSequence*> SelectedSequences = GetSelectedSequences();
-
-	for (UAvaSequence* const Sequence : SelectedSequences)
-	{
-		IAvaSequencePlaybackObject* const PlaybackObject = Provider.GetPlaybackObject();
-		if (Sequence && PlaybackObject)
-		{
-			PlaybackObject->ContinueSequence(Sequence);
-		}
-	}
-}
-
-bool FAvaSequencer::CanStopSequences() const
-{
-	return GetSelectedSequenceItems().Num() > 0;
-}
-
-void FAvaSequencer::StopSequences()
-{
-	const TArray<UAvaSequence*> SelectedSequences = GetSelectedSequences();
-
-	for (UAvaSequence* const Sequence : SelectedSequences)
-	{
-		IAvaSequencePlaybackObject* const PlaybackObject = Provider.GetPlaybackObject();
-		if (Sequence && PlaybackObject)
-		{
-			PlaybackObject->StopSequence(Sequence);
-		}
-	}
-}
-
 UObject* FAvaSequencer::FindObjectToPossess(UObject* InParentObject, const FMovieScenePossessable& InPossessable)
 {
 	if (!InParentObject)
@@ -2040,7 +1809,16 @@ TSharedRef<SWidget> FAvaSequencer::CreateSequenceWidget()
 		Sequencer->RegisterDrawerSection(SidebarDrawerId, MakeShared<FAvaSequenceSettingsDetails>(ThisSequencerRef));
 	}
 
-	RegisterNavigationTool();
+	if (UObject* const PlaybackContextObject = Provider.GetPlaybackContext())
+	{
+		if (UWorld* const World = PlaybackContextObject->GetWorld())
+		{
+			if (UAvaSequencerSubsystem* const SequencerSubsystem = World->GetSubsystem<UAvaSequencerSubsystem>())
+			{
+				SequencerSubsystem->OnSequencerCreated().Broadcast(SharedThis(this));
+			}
+		}
+	}
 
 	return Sequencer->GetSequencerWidget();
 }
@@ -2113,13 +1891,6 @@ void FAvaSequencer::OnEditorSelectionChanged(const FAvaEditorSelection& InEditor
 
 void FAvaSequencer::NotifyOnSequenceTreeChanged()
 {
-	using namespace UE::SequenceNavigator;
-
-	if (const TSharedPtr<INavigationTool> Tool = FNavigationToolExtender::FindNavigationTool(GetSequencer()))
-	{
-		Tool->RequestRefresh();
-	}
-
 	IAvaSequenceProvider* const SequenceProvider = Provider.GetSequenceProvider();
 	if (!SequenceTreeView.IsValid() || !SequenceProvider)
 	{
@@ -2206,47 +1977,9 @@ void FAvaSequencer::PostRedo(const bool bInSuccess)
 	PostUndo(bInSuccess);
 }
 
-TArray<TSharedPtr<UE::SequenceNavigator::FNavigationToolAvaSequence>> FAvaSequencer::GetSelectedSequenceItems() const
-{
-	using namespace UE::SequenceNavigator;
-
-	const TSharedPtr<ISequencer> Sequencer = SequencerWeak.Pin();
-	if (!Sequencer.IsValid())
-	{
-		return {};
-	}
-
-	const TSharedPtr<INavigationTool> Tool = FNavigationToolExtender::FindNavigationTool(Sequencer.ToSharedRef());
-	if (!Tool.IsValid())
-	{
-		return {};
-	}
-
-	const TArray<FNavigationToolItemPtr> SelectedItems = Tool->GetSelectedItems();
-	if (SelectedItems.IsEmpty())
-	{
-		return {};
-	}
-
-	TArray<TSharedPtr<FNavigationToolAvaSequence>> SelectedSequenceItems;
-	SelectedSequenceItems.Reserve(SelectedItems.Num());
-
-	for (const FNavigationToolItemPtr& Item : SelectedItems)
-	{
-		if (FNavigationToolAvaSequence* const AvaSequenceItem = Item->CastTo<FNavigationToolAvaSequence>())
-		{
-			SelectedSequenceItems.Add(SharedThis(AvaSequenceItem));
-		}
-	}
-
-	return SelectedSequenceItems;
-}
-
 TArray<UAvaSequence*> FAvaSequencer::GetSelectedSequences() const
 {
-	using namespace UE::SequenceNavigator;
-
-	const TArray<TSharedPtr<FNavigationToolAvaSequence>> SelectedItems = GetSelectedSequenceItems();
+	const TArray<TSharedPtr<IAvaSequenceItem>> SelectedItems = SequenceTreeView->GetSelectedItems();
 	if (SelectedItems.IsEmpty())
 	{
 		return {};
@@ -2255,9 +1988,9 @@ TArray<UAvaSequence*> FAvaSequencer::GetSelectedSequences() const
 	TArray<UAvaSequence*> SelectedAvaSequences;
 	SelectedAvaSequences.Reserve(SelectedItems.Num());
 
-	for (const TSharedPtr<FNavigationToolAvaSequence>& SequenceItem : SelectedItems)
+	for (const TSharedPtr<IAvaSequenceItem>& SequenceItem : SelectedItems)
 	{
-		if (UAvaSequence* const AvaSequence = SequenceItem->GetAvaSequence())
+		if (UAvaSequence* const AvaSequence = SequenceItem->GetSequence())
 		{
 			SelectedAvaSequences.Add(AvaSequence);
 		}
@@ -2442,52 +2175,4 @@ const TArray<TWeakObjectPtr<UAvaSequence>>& FAvaSequencer::GetRootSequences() co
 	return SequenceProvider->GetRootSequences();
 }
 
-void FAvaSequencer::RegisterNavigationTool()
-{
-	using namespace UE::SequenceNavigator;
-
-	if (!SequencerWeak.IsValid())
-	{
-		return;
-	}
-
-	const TSharedRef<ISequencer> Sequencer = SequencerWeak.Pin().ToSharedRef();
-	const FName ToolId = FNavigationToolExtender::GetToolInstanceId(*Sequencer);
-
-	const TSharedPtr<FNavigationToolProvider> FoundProvider = FNavigationToolExtender::FindToolProvider(ToolId
-		, UE::AvaSequencer::FAvaNavigationToolProvider::Identifier);
-	if (!FoundProvider.IsValid())
-	{
-		NavigationToolProvider = MakeShared<UE::AvaSequencer::FAvaNavigationToolProvider>(SharedThis(this));
-		FNavigationToolExtender::RegisterToolProvider(Sequencer, NavigationToolProvider.ToSharedRef());
-	}
-}
-
-void FAvaSequencer::UnregisterNavigationTool()
-{
-	using namespace UE::SequenceNavigator;
-
-	if (!NavigationToolProvider.IsValid())
-	{
-		return;
-	}
-
-	const TSharedPtr<ISequencer> Sequencer = SequencerWeak.Pin();
-	if (!Sequencer.IsValid())
-	{
-		return;
-	}
-
-	const TSharedPtr<INavigationTool> Tool = FNavigationToolExtender::FindNavigationTool(Sequencer.ToSharedRef());
-	if (!Tool.IsValid())
-	{
-		return;
-	}
-
-	const FName ToolId = FNavigationToolExtender::GetToolInstanceId(*Tool);
-	FNavigationToolExtender::UnregisterToolProvider(ToolId, NavigationToolProvider->GetIdentifier());
-
-	NavigationToolProvider.Reset();
-}
-
 #undef LOCTEXT_NAMESPACE
diff --git a/Engine/Plugins/VirtualProduction/Avalanche/Source/AvalancheSequencer/Private/AvaSequencer.h b/Engine/Plugins/VirtualProduction/Avalanche/Source/AvalancheSequencer/Private/AvaSequencer.h
--- a/Engine/Plugins/VirtualProduction/Avalanche/Source/AvalancheSequencer/Private/AvaSequencer.h
+++ b/Engine/Plugins/VirtualProduction/Avalanche/Source/AvalancheSequencer/Private/AvaSequencer.h
@@ -43,15 +43,9 @@ namespace UE
 		class SOutlinerView;
 	}
 
-	namespace AvaSequencer
-	{
-		class FAvaNavigationToolProvider;
-	}
-
 	namespace SequenceNavigator
 	{
 		class FNavigationToolAvaSequence;
-		class FNavigationToolProvider;
 	}
 }
 
@@ -196,36 +190,14 @@ public:
 	virtual void NotifyOnSequenceTreeChanged() override;
 	virtual const TArray<TWeakObjectPtr<UAvaSequence>>& GetRootSequences() const override;
 
-	//~ End IAvaSequencer
-
-	virtual bool CanAddSequence() const;
-	virtual UAvaSequence* AddSequence(UAvaSequence* const InParentSequence = nullptr);
-
-	virtual bool CanDuplicateSequences() const;
-	virtual void DuplicateSequences();
-
-	virtual bool CanExportSequences() const;
-	virtual void ExportSequences();
-
-	virtual bool CanSpawnPlayers() const;
-	virtual void SpawnPlayers();
+	virtual bool CanAddSequence() const override;
+	virtual UAvaSequence* AddSequence(UAvaSequence* const InParentSequence = nullptr) override;
 
-	virtual bool CanDeleteSequence() const;
-	virtual void DeleteSequence();
-	virtual void DeleteSequences(const TSet<UAvaSequence*>& InSequences);
+	virtual void DeleteSequences(const TSet<UAvaSequence*>& InSequences) override;
 
-	virtual bool CanRelabelSequence() const;
-	virtual void RelabelSequence();
-
-	virtual bool CanPlaySequences() const;
-	virtual void PlaySequences();
-
-	virtual bool CanContinueSequences() const;
-	virtual void ContinueSequences();
-
-	virtual bool CanStopSequences() const;
-	virtual void StopSequences();
+	//~ End IAvaSequencer
 
+	// BEGIN REMOVE AFTER OLD TREE VIEW REMOVAL
 	bool AddSequence_CanExecute() const;
 	void AddSequence_Execute();
 
@@ -252,6 +224,7 @@ public:
 
 	bool StopSequences_CanExecute() const;
 	void StopSequences_Execute();
+	// END REMOVE AFTER OLD TREE VIEW REMOVAL
 
 	/** Called when PostUndo has been called for the given sequence */
 	void OnSequenceEditUndo(UAvaSequence* InSequence);
@@ -262,8 +235,6 @@ public:
 	// End of FEditorUndoClient
 
 private:
-	TArray<TSharedPtr<UE::SequenceNavigator::FNavigationToolAvaSequence>> GetSelectedSequenceItems() const;
-
 	TArray<UAvaSequence*> GetSelectedSequences() const;
 
 	TSharedPtr<UE::Sequencer::SOutlinerView> GetOutlinerView() const;
@@ -281,9 +252,6 @@ private:
 
 	void OnSequencerClosed(const TSharedRef<ISequencer> InSequencer);
 
-	void RegisterNavigationTool();
-	void UnregisterNavigationTool();
-
 	FOnViewedSequenceChanged OnViewedSequenceChanged;
 
 	FOnSequenceAdded OnSequenceAddedDelegate;
@@ -349,6 +317,4 @@ private:
 	TSharedPtr<FExtender> SidebarExtender;
 
 	FDelegateHandle SidebarSelectionExtenderHandle;
-
-	TSharedPtr<UE::AvaSequencer::FAvaNavigationToolProvider> NavigationToolProvider;
 };
diff --git a/Engine/Plugins/VirtualProduction/Avalanche/Source/AvalancheSequencer/Private/NavigationTool/AvaNavigationToolProvider.cpp b/Engine/Plugins/VirtualProduction/Avalanche/Source/AvalancheSequencer/Private/NavigationTool/AvaNavigationToolProvider.cpp
deleted file mode 100644
diff --git a/Engine/Plugins/VirtualProduction/Avalanche/Source/AvalancheSequencer/Private/NavigationTool/AvaNavigationToolProvider.h b/Engine/Plugins/VirtualProduction/Avalanche/Source/AvalancheSequencer/Private/NavigationTool/AvaNavigationToolProvider.h
deleted file mode 100644
diff --git a/Engine/Plugins/VirtualProduction/Avalanche/Source/AvalancheSequencer/Private/NavigationTool/Columns/AvaNavigationToolStatusColumn.cpp b/Engine/Plugins/VirtualProduction/Avalanche/Source/AvalancheSequencer/Private/NavigationTool/Columns/AvaNavigationToolStatusColumn.cpp
deleted file mode 100644
diff --git a/Engine/Plugins/VirtualProduction/Avalanche/Source/AvalancheSequencer/Private/NavigationTool/Columns/AvaNavigationToolStatusColumn.h b/Engine/Plugins/VirtualProduction/Avalanche/Source/AvalancheSequencer/Private/NavigationTool/Columns/AvaNavigationToolStatusColumn.h
deleted file mode 100644
diff --git a/Engine/Plugins/VirtualProduction/Avalanche/Source/AvalancheSequencer/Private/NavigationTool/Columns/SAvaNavigationToolStatus.cpp b/Engine/Plugins/VirtualProduction/Avalanche/Source/AvalancheSequencer/Private/NavigationTool/Columns/SAvaNavigationToolStatus.cpp
deleted file mode 100644
diff --git a/Engine/Plugins/VirtualProduction/Avalanche/Source/AvalancheSequencer/Private/NavigationTool/Columns/SAvaNavigationToolStatus.h b/Engine/Plugins/VirtualProduction/Avalanche/Source/AvalancheSequencer/Private/NavigationTool/Columns/SAvaNavigationToolStatus.h
deleted file mode 100644
diff --git a/Engine/Plugins/VirtualProduction/Avalanche/Source/AvalancheSequencer/Private/NavigationTool/DragDrop/NavigationToolAvaSequenceDropHandler.cpp b/Engine/Plugins/VirtualProduction/Avalanche/Source/AvalancheSequencer/Private/NavigationTool/DragDrop/NavigationToolAvaSequenceDropHandler.cpp
deleted file mode 100644
diff --git a/Engine/Plugins/VirtualProduction/Avalanche/Source/AvalancheSequencer/Private/NavigationTool/DragDrop/NavigationToolAvaSequenceDropHandler.h b/Engine/Plugins/VirtualProduction/Avalanche/Source/AvalancheSequencer/Private/NavigationTool/DragDrop/NavigationToolAvaSequenceDropHandler.h
deleted file mode 100644
diff --git a/Engine/Plugins/VirtualProduction/Avalanche/Source/AvalancheSequencer/Private/NavigationTool/NavigationToolAvaSequence.cpp b/Engine/Plugins/VirtualProduction/Avalanche/Source/AvalancheSequencer/Private/NavigationTool/NavigationToolAvaSequence.cpp
deleted file mode 100644
diff --git a/Engine/Plugins/VirtualProduction/Avalanche/Source/AvalancheSequencer/Private/NavigationTool/NavigationToolAvaSequence.h b/Engine/Plugins/VirtualProduction/Avalanche/Source/AvalancheSequencer/Private/NavigationTool/NavigationToolAvaSequence.h
deleted file mode 100644
diff --git a/Engine/Plugins/VirtualProduction/Avalanche/Source/AvalancheSequencer/Private/SequenceTree/Widgets/SAvaSequenceTree.cpp b/Engine/Plugins/VirtualProduction/Avalanche/Source/AvalancheSequencer/Private/SequenceTree/Widgets/SAvaSequenceTree.cpp
--- a/Engine/Plugins/VirtualProduction/Avalanche/Source/AvalancheSequencer/Private/SequenceTree/Widgets/SAvaSequenceTree.cpp
+++ b/Engine/Plugins/VirtualProduction/Avalanche/Source/AvalancheSequencer/Private/SequenceTree/Widgets/SAvaSequenceTree.cpp
@@ -324,7 +324,7 @@ FReply SAvaSequenceTree::OnNewSequenceClicked()
 {
 	if (SequencerWeak.IsValid())
 	{
-		SequencerWeak.Pin()->AddSequence();
+		SequencerWeak.Pin()->AddSequence_Execute();
 		return FReply::Handled();
 	}
 	return FReply::Unhandled();
diff --git a/Engine/Plugins/VirtualProduction/Avalanche/Source/AvalancheSequencer/Private/Settings/AvaMarkSetting.h b/Engine/Plugins/VirtualProduction/Avalanche/Source/AvalancheSequencer/Private/Settings/AvaMarkSetting.h
deleted file mode 100644
diff --git a/Engine/Plugins/VirtualProduction/Avalanche/Source/AvalancheSequencer/Private/Settings/AvaSequencePreset.cpp b/Engine/Plugins/VirtualProduction/Avalanche/Source/AvalancheSequencer/Private/Settings/AvaSequencePreset.cpp
--- a/Engine/Plugins/VirtualProduction/Avalanche/Source/AvalancheSequencer/Private/Settings/AvaSequencePreset.cpp
+++ b/Engine/Plugins/VirtualProduction/Avalanche/Source/AvalancheSequencer/Private/Settings/AvaSequencePreset.cpp
@@ -1,6 +1,6 @@
 // Copyright Epic Games, Inc. All Rights Reserved.
 
-#include "AvaSequencePreset.h"
+#include "Settings/AvaSequencePreset.h"
 #include "AvaSequence.h"
 
 bool FAvaSequencePreset::ShouldModifySequence() const
diff --git a/Engine/Plugins/VirtualProduction/Avalanche/Source/AvalancheSequencer/Private/Settings/AvaSequencePreset.h b/Engine/Plugins/VirtualProduction/Avalanche/Source/AvalancheSequencer/Private/Settings/AvaSequencePreset.h
deleted file mode 100644
diff --git a/Engine/Plugins/VirtualProduction/Avalanche/Source/AvalancheSequencer/Private/Settings/AvaSequencerDisplayRate.h b/Engine/Plugins/VirtualProduction/Avalanche/Source/AvalancheSequencer/Private/Settings/AvaSequencerDisplayRate.h
deleted file mode 100644
diff --git a/Engine/Plugins/VirtualProduction/Avalanche/Source/AvalancheSequencer/Private/Settings/AvaSequencerSettings.cpp b/Engine/Plugins/VirtualProduction/Avalanche/Source/AvalancheSequencer/Private/Settings/AvaSequencerSettings.cpp
--- a/Engine/Plugins/VirtualProduction/Avalanche/Source/AvalancheSequencer/Private/Settings/AvaSequencerSettings.cpp
+++ b/Engine/Plugins/VirtualProduction/Avalanche/Source/AvalancheSequencer/Private/Settings/AvaSequencerSettings.cpp
@@ -1,6 +1,6 @@
 // Copyright Epic Games, Inc. All Rights Reserved.
 
-#include "AvaSequencerSettings.h"
+#include "Settings/AvaSequencerSettings.h"
 #include "AvaSequenceDefaultTags.h"
 #include "Misc/FrameRate.h"
 
diff --git a/Engine/Plugins/VirtualProduction/Avalanche/Source/AvalancheSequencer/Private/Settings/AvaSequencerSettings.h b/Engine/Plugins/VirtualProduction/Avalanche/Source/AvalancheSequencer/Private/Settings/AvaSequencerSettings.h
deleted file mode 100644
diff --git a/Engine/Plugins/VirtualProduction/Avalanche/Source/AvalancheSequencer/Public/AvaSequencerSubsystem.h b/Engine/Plugins/VirtualProduction/Avalanche/Source/AvalancheSequencer/Public/AvaSequencerSubsystem.h
--- a/Engine/Plugins/VirtualProduction/Avalanche/Source/AvalancheSequencer/Public/AvaSequencerSubsystem.h
+++ b/Engine/Plugins/VirtualProduction/Avalanche/Source/AvalancheSequencer/Public/AvaSequencerSubsystem.h
@@ -28,10 +28,16 @@ public:
 
 	TSharedPtr<IAvaSequencer> GetSequencer() const { return SequencerWeak.Pin(); }
 
+	/** A delegate that gets executed then a sequencer is created */
+	DECLARE_MULTICAST_DELEGATE_OneParam(FOnAvaSequencerCreated, TSharedRef<IAvaSequencer>);
+	FOnAvaSequencerCreated& OnSequencerCreated() { return OnSequencerCreatedDelegate; }
+
 	//~ Begin UWorldSubsystem
 	virtual bool DoesSupportWorldType(const EWorldType::Type InWorldType) const override;
 	//~ End UWorldSubsystem
 
 private:
 	TWeakPtr<IAvaSequencer> SequencerWeak;
+
+	FOnAvaSequencerCreated OnSequencerCreatedDelegate;
 };
diff --git a/Engine/Plugins/VirtualProduction/Avalanche/Source/AvalancheSequencer/Public/AvaSequencerUtils.h b/Engine/Plugins/VirtualProduction/Avalanche/Source/AvalancheSequencer/Public/AvaSequencerUtils.h
--- a/Engine/Plugins/VirtualProduction/Avalanche/Source/AvalancheSequencer/Public/AvaSequencerUtils.h
+++ b/Engine/Plugins/VirtualProduction/Avalanche/Source/AvalancheSequencer/Public/AvaSequencerUtils.h
@@ -15,7 +15,6 @@ class ISequencer;
 class UAvaSceneSubsystem;
 class UAvaSequencerSubsystem;
 class UWorld;
-enum class EMovieSceneTransformChannel : uint32;
 
 struct FAvaSequencerUtils
 {
@@ -47,11 +46,11 @@ struct FAvaSequencerUtils
 	static UAvaSceneSubsystem* GetSceneSubsystem(const TSharedRef<ISequencer>& InSequencer);
 
 	/** @return Motion Design Scene Interface from an ISequencer */
-	static IAvaSceneInterface* GetSceneInterface(const TSharedRef<ISequencer>& InSequencer);
+	AVALANCHESEQUENCER_API static IAvaSceneInterface* GetSceneInterface(const TSharedRef<ISequencer>& InSequencer);
 
 	/** @return Motion Design Sequence Provider from an ISequencer */
-	static IAvaSequenceProvider* GetSequenceProvider(const TSharedRef<ISequencer>& InSequencer);
+	AVALANCHESEQUENCER_API static IAvaSequenceProvider* GetSequenceProvider(const TSharedRef<ISequencer>& InSequencer);
 
 	/** @return Motion Design Sequencer from an ISequencer */
-	static TSharedPtr<IAvaSequencer> GetAvaSequencer(const TSharedRef<ISequencer>& InSequencer);
+	AVALANCHESEQUENCER_API static TSharedPtr<IAvaSequencer> GetAvaSequencer(const TSharedRef<ISequencer>& InSequencer);
 };
diff --git a/Engine/Plugins/VirtualProduction/Avalanche/Source/AvalancheSequencer/Public/IAvaSequencer.h b/Engine/Plugins/VirtualProduction/Avalanche/Source/AvalancheSequencer/Public/IAvaSequencer.h
--- a/Engine/Plugins/VirtualProduction/Avalanche/Source/AvalancheSequencer/Public/IAvaSequencer.h
+++ b/Engine/Plugins/VirtualProduction/Avalanche/Source/AvalancheSequencer/Public/IAvaSequencer.h
@@ -18,6 +18,7 @@ class IAvaSequencerProvider;
 class IPropertyHandle;
 class ISequencer;
 class SWidget;
+class UAvaSequence;
 class UObject;
 class USequencerSettings;
 
@@ -69,4 +70,9 @@ public:
 
 	DECLARE_MULTICAST_DELEGATE_OneParam(FOnSequenceRemoved, UAvaSequence*);
 	virtual FOnSequenceRemoved& OnSequenceRemoved() = 0;
+
+	virtual bool CanAddSequence() const = 0;
+	virtual UAvaSequence* AddSequence(UAvaSequence* const InParentSequence = nullptr) = 0;
+
+	virtual void DeleteSequences(const TSet<UAvaSequence*>& InSequences) = 0;
 };
diff --git a/Engine/Plugins/VirtualProduction/Avalanche/Source/AvalancheSequencer/Public/IAvaSequencerProvider.h b/Engine/Plugins/VirtualProduction/Avalanche/Source/AvalancheSequencer/Public/IAvaSequencerProvider.h
--- a/Engine/Plugins/VirtualProduction/Avalanche/Source/AvalancheSequencer/Public/IAvaSequencerProvider.h
+++ b/Engine/Plugins/VirtualProduction/Avalanche/Source/AvalancheSequencer/Public/IAvaSequencerProvider.h
@@ -9,6 +9,7 @@
 
 class FEditorModeTools;
 class FEditorViewportClient;
+class IAvaSequencer;
 class IAvaSequencePlaybackObject;
 class IAvaSequenceProvider;
 class ISequencer;
@@ -22,6 +23,8 @@ class IAvaSequencerProvider
 public:
     virtual ~IAvaSequencerProvider() = default;
 
+	virtual TSharedPtr<IAvaSequencer> GetAvaSequencer() const = 0;
+
 	/** Option for FAvaSequencer to use an externally instanced ISequencer instead of instancing its own */
 	virtual TSharedPtr<ISequencer> GetExternalSequencer() const { return nullptr; };
 
diff --git a/Engine/Plugins/VirtualProduction/Avalanche/Source/AvalancheSequencer/Public/Settings/AvaMarkSetting.h b/Engine/Plugins/VirtualProduction/Avalanche/Source/AvalancheSequencer/Public/Settings/AvaMarkSetting.h
new file mode 100644
--- /dev/null
+++ b/Engine/Plugins/VirtualProduction/Avalanche/Source/AvalancheSequencer/Public/Settings/AvaMarkSetting.h
@@ -0,0 +1,22 @@
+// Copyright Epic Games, Inc. All Rights Reserved.
+
+#pragma once
+
+#include "Marks/AvaMark.h"
+#include "MovieSceneMarkedFrame.h"
+#include "AvaMarkSetting.generated.h"
+
+USTRUCT()
+struct FAvaMarkSetting
+{
+	GENERATED_BODY()
+
+	UPROPERTY(EditAnywhere, Category = "Mark")
+	FString Label;
+
+	UPROPERTY(EditAnywhere, Category = "Mark")
+	int32 FrameNumber = 0;
+
+	UPROPERTY(EditAnywhere, Category = "Mark")
+	FAvaMark Mark;
+};
diff --git a/Engine/Plugins/VirtualProduction/Avalanche/Source/AvalancheSequencer/Public/Settings/AvaSequencePreset.h b/Engine/Plugins/VirtualProduction/Avalanche/Source/AvalancheSequencer/Public/Settings/AvaSequencePreset.h
new file mode 100644
--- /dev/null
+++ b/Engine/Plugins/VirtualProduction/Avalanche/Source/AvalancheSequencer/Public/Settings/AvaSequencePreset.h
@@ -0,0 +1,63 @@
+// Copyright Epic Games, Inc. All Rights Reserved.
+
+#pragma once
+
+#include "AvaMarkSetting.h"
+#include "AvaTagSoftHandle.h"
+#include "AvaSequencePreset.generated.h"
+
+class UAvaSequence;
+
+/** Simple settings to apply to a sequence */
+USTRUCT()
+struct FAvaSequencePreset
+{
+	GENERATED_BODY()
+
+	AVALANCHESEQUENCER_API FAvaSequencePreset() = default;
+	AVALANCHESEQUENCER_API virtual ~FAvaSequencePreset() = default;
+
+	AVALANCHESEQUENCER_API explicit FAvaSequencePreset(FName InPresetName)
+		: PresetName(InPresetName)
+	{
+	}
+
+	bool ShouldModifySequence() const;
+
+	bool ShouldModifyMovieScene() const;
+
+	AVALANCHESEQUENCER_API void ApplyPreset(UAvaSequence* InSequence) const;
+
+	bool operator==(const FAvaSequencePreset& InOther) const
+	{
+		return PresetName == InOther.PresetName;
+	}
+
+	friend uint32 GetTypeHash(const FAvaSequencePreset& InPreset)
+	{
+		return GetTypeHash(InPreset.PresetName);
+	}
+
+	/** Name to identify this preset */
+	UPROPERTY(EditAnywhere, Category="Sequence")
+	FName PresetName;
+
+	/** If not none, the sequence label to set */
+	UPROPERTY(EditAnywhere, Category="Sequence", meta=(EditCondition="bEnableLabel"))
+	FName SequenceLabel;
+
+	UPROPERTY(EditAnywhere, Category="Sequence", meta=(EditCondition="bEnableTag"))
+	FAvaTagSoftHandle SequenceTag;
+
+	UPROPERTY(EditAnywhere, Category="Sequence", meta=(EditCondition="bEnableMarks"))
+	TArray<FAvaMarkSetting> Marks;
+
+	UPROPERTY(meta = (InlineEditConditionToggle))
+	bool bEnableLabel = false;
+
+	UPROPERTY(meta = (InlineEditConditionToggle))
+	bool bEnableTag = false;
+
+	UPROPERTY(meta = (InlineEditConditionToggle))
+	bool bEnableMarks = false;
+};
diff --git a/Engine/Plugins/VirtualProduction/Avalanche/Source/AvalancheSequencer/Public/Settings/AvaSequencerDisplayRate.h b/Engine/Plugins/VirtualProduction/Avalanche/Source/AvalancheSequencer/Public/Settings/AvaSequencerDisplayRate.h
new file mode 100644
--- /dev/null
+++ b/Engine/Plugins/VirtualProduction/Avalanche/Source/AvalancheSequencer/Public/Settings/AvaSequencerDisplayRate.h
@@ -0,0 +1,15 @@
+// Copyright Epic Games, Inc. All Rights Reserved.
+
+#pragma once
+
+#include "Misc/FrameRate.h"
+#include "AvaSequencerDisplayRate.generated.h"
+
+USTRUCT(BlueprintType)
+struct FAvaSequencerDisplayRate
+{
+	GENERATED_BODY()
+
+	UPROPERTY()
+	FFrameRate FrameRate;
+};
diff --git a/Engine/Plugins/VirtualProduction/Avalanche/Source/AvalancheSequencer/Public/Settings/AvaSequencerSettings.h b/Engine/Plugins/VirtualProduction/Avalanche/Source/AvalancheSequencer/Public/Settings/AvaSequencerSettings.h
new file mode 100644
--- /dev/null
+++ b/Engine/Plugins/VirtualProduction/Avalanche/Source/AvalancheSequencer/Public/Settings/AvaSequencerSettings.h
@@ -0,0 +1,75 @@
+// Copyright Epic Games, Inc. All Rights Reserved.
+
+#pragma once
+
+#include "AvaSequencePreset.h"
+#include "AvaSequencerDisplayRate.h"
+#include "Engine/DeveloperSettings.h"
+#include "Sidebar/SidebarState.h"
+#include "UObject/Object.h"
+#include "AvaSequencerSettings.generated.h"
+
+UCLASS(config=EditorPerProjectUserSettings, meta = (DisplayName = "Sequencer"))
+class UAvaSequencerSettings : public UDeveloperSettings
+{
+	GENERATED_BODY()
+
+public:
+	static constexpr const TCHAR* SettingsCategory = TEXT("Motion Design");
+	static constexpr const TCHAR* SettingsSection  = TEXT("Sequencer");
+
+	UAvaSequencerSettings();
+
+	FFrameRate GetDisplayRate() const
+	{
+		return DisplayRate.FrameRate;
+	}
+
+	double GetStartTime() const
+	{
+		return StartTime;
+	}
+
+	double GetEndTime() const
+	{
+		return EndTime;
+	}
+
+	AVALANCHESEQUENCER_API TConstArrayView<FAvaSequencePreset> GetDefaultSequencePresets() const;
+
+	const TSet<FAvaSequencePreset>& GetCustomSequencePresets() const
+	{
+		return CustomSequencePresets;
+	}
+
+	FSidebarState& GetSidebarState()
+	{
+		return SidebarState;
+	}
+
+	void SetSidebarState(const FSidebarState& InSidebarState)
+	{
+		SidebarState = InSidebarState;
+	}
+
+private:
+	/** The default display rate to use for new sequences */
+	UPROPERTY(Config, EditAnywhere, Category = "Playback")
+	FAvaSequencerDisplayRate DisplayRate;
+
+	/** The default start time to use for new sequences */
+	UPROPERTY(Config, EditAnywhere, Category = "Playback")
+	double StartTime = 0.0;
+
+	/** The default end time to use for new sequences */
+	UPROPERTY(Config, EditAnywhere, Category = "Playback")
+	double EndTime = 2.0;
+
+	/** Sequence Presets that are uniquely identified by their Preset Name */
+	UPROPERTY(Config, EditAnywhere, Category = "Sequencer")
+	TSet<FAvaSequencePreset> CustomSequencePresets;
+
+	/** The state of a sidebar to be restored when Sequencer is initialized */
+	UPROPERTY(Config)
+	FSidebarState SidebarState;
+};
