[SequenceNavigator] Refactor to use Sequencer view models and type macros

diff --git a/Engine/Plugins/Experimental/AvalancheSequenceNavigator/Source/AvalancheSequenceNavigator/AvalancheSequenceNavigator.Build.cs b/Engine/Plugins/Experimental/AvalancheSequenceNavigator/Source/AvalancheSequenceNavigator/AvalancheSequenceNavigator.Build.cs
--- a/Engine/Plugins/Experimental/AvalancheSequenceNavigator/Source/AvalancheSequenceNavigator/AvalancheSequenceNavigator.Build.cs
+++ b/Engine/Plugins/Experimental/AvalancheSequenceNavigator/Source/AvalancheSequenceNavigator/AvalancheSequenceNavigator.Build.cs
@@ -20,8 +20,9 @@ public class AvalancheSequenceNavigator : ModuleRules
 			"CoreUObject",
 			"Engine",
 			"MovieScene",
-			"Sequencer",
 			"SequenceNavigator",
+			"Sequencer",
+			"SequencerCore",
 			"Slate",
 			"SlateCore",
 			"ToolMenus",
diff --git a/Engine/Plugins/Experimental/AvalancheSequenceNavigator/Source/AvalancheSequenceNavigator/Private/AvaNavigationToolProvider.cpp b/Engine/Plugins/Experimental/AvalancheSequenceNavigator/Source/AvalancheSequenceNavigator/Private/AvaNavigationToolProvider.cpp
--- a/Engine/Plugins/Experimental/AvalancheSequenceNavigator/Source/AvalancheSequenceNavigator/Private/AvaNavigationToolProvider.cpp
+++ b/Engine/Plugins/Experimental/AvalancheSequenceNavigator/Source/AvalancheSequenceNavigator/Private/AvaNavigationToolProvider.cpp
@@ -220,14 +220,14 @@ void FAvaNavigationToolProvider::OnExtendColumnViews(TSet<FNavigationToolColumnV
 	FNavigationToolProvider::OnExtendColumnViews(OutColumnViews);
 }
 
-void FAvaNavigationToolProvider::OnExtendItemChildren(UE::SequenceNavigator::INavigationTool& InTool
-	, const UE::SequenceNavigator::FNavigationToolItemPtr& InParentItem
-	, TArray<UE::SequenceNavigator::FNavigationToolItemPtr>& OutChildren
+void FAvaNavigationToolProvider::OnExtendItemChildren(SequenceNavigator::INavigationTool& InTool
+	, const SequenceNavigator::FNavigationToolViewModelPtr& InParentItem
+	, TArray<SequenceNavigator::FNavigationToolViewModelWeakPtr>& OutWeakChildren
 	, const bool bInRecursive)
 {
-	using namespace UE::SequenceNavigator;
+	using namespace SequenceNavigator;
 
-	FNavigationToolProvider::OnExtendItemChildren(InTool, InParentItem, OutChildren, bInRecursive);
+	FNavigationToolProvider::OnExtendItemChildren(InTool, InParentItem, OutWeakChildren, bInRecursive);
 
 	// Only extending root item
 	if (InParentItem->GetItemId() != FNavigationToolItemId::RootId)
@@ -253,20 +253,20 @@ void FAvaNavigationToolProvider::OnExtendItemChildren(UE::SequenceNavigator::INa
 			return InA->GetDisplayName().CompareTo(InB->GetDisplayName()) > 0;
 		});
 
-	const TSharedRef<FAvaNavigationToolProvider> Provider = SharedThis(this);
+	const TSharedRef<FAvaNavigationToolProvider> ProviderRef = SharedThis(this);
 
 	// Add child sequence items
 	for (const TWeakObjectPtr<UAvaSequence>& RootSequenceWeak : RootWeakSequences)
 	{
 		if (UAvaSequence* const RootSequence = RootSequenceWeak.Get())
 		{
-			const FNavigationToolItemPtr NewItem = InTool.FindOrAdd<FNavigationToolAvaSequence>(Provider
+			const FNavigationToolViewModelPtr NewItem = InTool.FindOrAdd<FNavigationToolAvaSequence>(ProviderRef
 				, InParentItem, RootSequence);
 			const FNavigationToolItemFlagGuard Guard(NewItem, ENavigationToolItemFlags::IgnorePendingKill);
-			OutChildren.Add(NewItem);
+			OutWeakChildren.Add(NewItem);
 			if (bInRecursive)
 			{
-				NewItem->FindChildren(OutChildren, bInRecursive);
+				NewItem->FindChildren(OutWeakChildren, bInRecursive);
 			}
 		}
 	}
@@ -438,14 +438,14 @@ void FAvaNavigationToolProvider::RemoveToolItemContextMenuExtension()
 	}
 }
 
-void FAvaNavigationToolProvider::ExtendToolItemDragDropOp(UE::SequenceNavigator::FNavigationToolItemDragDropOp& InDragDropOp)
+void FAvaNavigationToolProvider::ExtendToolItemDragDropOp(SequenceNavigator::FNavigationToolItemDragDropOp& InDragDropOp)
 {
-	InDragDropOp.AddDropHandler<UE::SequenceNavigator::FNavigationToolAvaSequenceDropHandler>(WeakAvaSequencer);
+	InDragDropOp.AddDropHandler<SequenceNavigator::FNavigationToolAvaSequenceDropHandler>(WeakAvaSequencer);
 }
 
 void FAvaNavigationToolProvider::OnSequenceAdded(UAvaSequence* const InAvaSequence)
 {
-	using namespace UE::SequenceNavigator;
+	using namespace SequenceNavigator;
 
 	const TSharedPtr<IAvaSequencer> AvaSequencer = WeakAvaSequencer.Pin();
 	if (!AvaSequencer.IsValid())
@@ -531,9 +531,10 @@ void FAvaNavigationToolProvider::GeneratePresetMenu(UToolMenu* const InToolMenu)
 		}));
 }
 
-TArray<TSharedPtr<UE::SequenceNavigator::FNavigationToolAvaSequence>> FAvaNavigationToolProvider::GetSelectedSequenceItems() const
+TArray<Sequencer::TViewModelPtr<SequenceNavigator::FNavigationToolAvaSequence>> FAvaNavigationToolProvider::GetSelectedSequenceItems() const
 {
-	using namespace UE::SequenceNavigator;
+	using namespace SequenceNavigator;
+	using namespace Sequencer;
 
 	const TSharedPtr<INavigationTool> Tool = GetNavigationTool();
 	if (!Tool.IsValid())
@@ -541,20 +542,20 @@ TArray<TSharedPtr<UE::SequenceNavigator::FNavigationToolAvaSequence>> FAvaNaviga
 		return {};
 	}
 
-	const TArray<FNavigationToolItemPtr> SelectedItems = Tool->GetSelectedItems();
-	if (SelectedItems.IsEmpty())
+	const TArray<FNavigationToolViewModelWeakPtr> WeakSelectedItems = Tool->GetSelectedItems();
+	if (WeakSelectedItems.IsEmpty())
 	{
 		return {};
 	}
 
-	TArray<TSharedPtr<FNavigationToolAvaSequence>> SelectedSequenceItems;
-	SelectedSequenceItems.Reserve(SelectedItems.Num());
+	TArray<TViewModelPtr<FNavigationToolAvaSequence>> SelectedSequenceItems;
+	SelectedSequenceItems.Reserve(WeakSelectedItems.Num());
 
-	for (const FNavigationToolItemPtr& Item : SelectedItems)
+	for (const FNavigationToolViewModelWeakPtr& WeakItem : WeakSelectedItems)
 	{
-		if (FNavigationToolAvaSequence* const AvaSequenceItem = Item->CastTo<FNavigationToolAvaSequence>())
+		if (const TViewModelPtr<FNavigationToolAvaSequence> AvaSequenceItem = WeakItem.ImplicitPin())
 		{
-			SelectedSequenceItems.Add(SharedThis(AvaSequenceItem));
+			SelectedSequenceItems.Add(AvaSequenceItem);
 		}
 	}
 
@@ -564,8 +565,9 @@ TArray<TSharedPtr<UE::SequenceNavigator::FNavigationToolAvaSequence>> FAvaNaviga
 TArray<UAvaSequence*> FAvaNavigationToolProvider::GetSelectedSequences() const
 {
 	using namespace UE::SequenceNavigator;
+	using namespace UE::Sequencer;
 
-	const TArray<TSharedPtr<FNavigationToolAvaSequence>> SelectedItems = GetSelectedSequenceItems();
+	const TArray<TViewModelPtr<FNavigationToolAvaSequence>> SelectedItems = GetSelectedSequenceItems();
 	if (SelectedItems.IsEmpty())
 	{
 		return {};
@@ -593,6 +595,7 @@ bool FAvaNavigationToolProvider::CanRelabelSelection() const
 void FAvaNavigationToolProvider::RelabelSelection()
 {
 	using namespace UE::SequenceNavigator;
+	using namespace UE::Sequencer;
 
 	const TSharedPtr<IAvaSequencer> AvaSequencer = WeakAvaSequencer.Pin();
 	if (!AvaSequencer.IsValid())
@@ -606,13 +609,19 @@ void FAvaNavigationToolProvider::RelabelSelection()
 		return;
 	}
 
-	const TArray<FNavigationToolItemPtr> SelectedItems = NavigationTool->GetSelectedItems();
-	if (SelectedItems.Num() != 1)
+	const TArray<FNavigationToolViewModelWeakPtr> WeakSelectedItems = NavigationTool->GetSelectedItems();
+	if (WeakSelectedItems.Num() != 1)
+	{
+		return;
+	}
+
+	const FNavigationToolViewModelPtr FirstSelectedItem = WeakSelectedItems[0].Pin();
+	if (!FirstSelectedItem.IsValid())
 	{
 		return;
 	}
 
-	const FNavigationToolAvaSequence* const AvaSequenceItem = SelectedItems[0]->CastTo<FNavigationToolAvaSequence>();
+	const TViewModelPtr<FNavigationToolAvaSequence> AvaSequenceItem = FirstSelectedItem.ImplicitCast();
 	if (!AvaSequenceItem)
 	{
 		return;
@@ -630,6 +639,7 @@ bool FAvaNavigationToolProvider::CanAddSequenceToSelection() const
 void FAvaNavigationToolProvider::AddSequenceToSelection()
 {
 	using namespace UE::SequenceNavigator;
+	using namespace UE::Sequencer;
 
 	const TSharedPtr<IAvaSequencer> AvaSequencer = WeakAvaSequencer.Pin();
 	if (!AvaSequencer.IsValid())
@@ -643,10 +653,10 @@ void FAvaNavigationToolProvider::AddSequenceToSelection()
 		return;
 	}
 
-	TSharedPtr<FNavigationToolAvaSequence> ParentSequenceItem;
+	TViewModelPtr<FNavigationToolAvaSequence> ParentSequenceItem;
 	UAvaSequence* ParentSequence = nullptr;
 
-	const TArray<TSharedPtr<FNavigationToolAvaSequence>> SelectedSequences = GetSelectedSequenceItems();
+	const TArray<TViewModelPtr<FNavigationToolAvaSequence>> SelectedSequences = GetSelectedSequenceItems();
 	if (SelectedSequences.Num() == 1)
 	{
 		ParentSequenceItem = SelectedSequences[0];
@@ -659,11 +669,15 @@ void FAvaNavigationToolProvider::AddSequenceToSelection()
 		return;
 	}
 
+	const FNavigationToolViewModelPtr ActualParentItem = ParentSequenceItem.IsValid()
+		? ParentSequenceItem : NavigationTool->GetTreeRoot().ImplicitPin();
+
+	const FNavigationToolViewModelPtr NewItem = NavigationTool->FindOrAdd<FNavigationToolAvaSequence>(SharedThis(this)
+		, ActualParentItem, NewSequence);
+
 	FNavigationToolAddItemParams AddChildItemParams;
-	AddChildItemParams.Item = NavigationTool->FindOrAdd<FNavigationToolAvaSequence>(SharedThis(this)
-		, AddChildItemParams.RelativeItem, NewSequence);
-	AddChildItemParams.RelativeItem = ParentSequenceItem.IsValid()
-		? ParentSequenceItem.ToSharedRef() : NavigationTool->GetTreeRoot();
+	AddChildItemParams.WeakItem = NewItem;
+	AddChildItemParams.WeakRelativeItem = ActualParentItem.AsWeak();
 	AddChildItemParams.RelativeDropZone = EItemDropZone::OntoItem;
 	AddChildItemParams.Flags = ENavigationToolAddItemFlags::Select | ENavigationToolAddItemFlags::Transact;
 	AddChildItemParams.SelectionFlags = ENavigationToolItemSelectionFlags::SignalSelectionChange
diff --git a/Engine/Plugins/Experimental/AvalancheSequenceNavigator/Source/AvalancheSequenceNavigator/Private/AvaNavigationToolProvider.h b/Engine/Plugins/Experimental/AvalancheSequenceNavigator/Source/AvalancheSequenceNavigator/Private/AvaNavigationToolProvider.h
--- a/Engine/Plugins/Experimental/AvalancheSequenceNavigator/Source/AvalancheSequenceNavigator/Private/AvaNavigationToolProvider.h
+++ b/Engine/Plugins/Experimental/AvalancheSequenceNavigator/Source/AvalancheSequenceNavigator/Private/AvaNavigationToolProvider.h
@@ -60,8 +60,8 @@ public:
 	virtual void OnExtendColumns(SequenceNavigator::FNavigationToolColumnExtender& OutExtender) override;
 	virtual void OnExtendColumnViews(TSet<FNavigationToolColumnView>& OutColumnViews) override;
 	virtual void OnExtendItemChildren(SequenceNavigator::INavigationTool& InTool
-		, const SequenceNavigator::FNavigationToolItemPtr& InParentItem
-		, TArray<SequenceNavigator::FNavigationToolItemPtr>& OutChildren
+		, const SequenceNavigator::FNavigationToolViewModelPtr& InParentItem
+		, TArray<SequenceNavigator::FNavigationToolViewModelWeakPtr>& OutWeakChildren
 		, const bool bInRecursive) override;
 	virtual void OnExtendBuiltInFilters(TArray<FNavigationToolBuiltInFilterParams>& OutFilterParams) override;
 
@@ -91,7 +91,7 @@ private:
 
 	void GeneratePresetMenu(UToolMenu* const InToolMenu);
 
-	TArray<TSharedPtr<SequenceNavigator::FNavigationToolAvaSequence>> GetSelectedSequenceItems() const;
+	TArray<Sequencer::TViewModelPtr<SequenceNavigator::FNavigationToolAvaSequence>> GetSelectedSequenceItems() const;
 	TArray<UAvaSequence*> GetSelectedSequences() const;
 
 	bool CanRelabelSelection() const;
diff --git a/Engine/Plugins/Experimental/AvalancheSequenceNavigator/Source/AvalancheSequenceNavigator/Private/Columns/AvaNavigationToolStatusColumn.cpp b/Engine/Plugins/Experimental/AvalancheSequenceNavigator/Source/AvalancheSequenceNavigator/Private/Columns/AvaNavigationToolStatusColumn.cpp
--- a/Engine/Plugins/Experimental/AvalancheSequenceNavigator/Source/AvalancheSequenceNavigator/Private/Columns/AvaNavigationToolStatusColumn.cpp
+++ b/Engine/Plugins/Experimental/AvalancheSequenceNavigator/Source/AvalancheSequenceNavigator/Private/Columns/AvaNavigationToolStatusColumn.cpp
@@ -13,6 +13,8 @@
 namespace UE::SequenceNavigator
 {
 
+UE_SEQUENCER_DEFINE_CASTABLE(FAvaNavigationToolStatusColumn)
+
 FText FAvaNavigationToolStatusColumn::GetColumnDisplayNameText() const
 {
 	return LOCTEXT("StatusColumn", "Status");
@@ -37,11 +39,11 @@ SHeaderRow::FColumn::FArguments FAvaNavigationToolStatusColumn::ConstructHeaderR
 		.OnGetMenuContent(InToolView, &INavigationToolView::GetColumnMenuContent, ColumnId);
 }
 
-TSharedRef<SWidget> FAvaNavigationToolStatusColumn::ConstructRowWidget(const FNavigationToolItemRef& InItem
+TSharedRef<SWidget> FAvaNavigationToolStatusColumn::ConstructRowWidget(const FNavigationToolViewModelPtr& InItem
 	, const TSharedRef<INavigationToolView>& InView
 	, const TSharedRef<SNavigationToolTreeRow>& InRow)
 {
-	if (!InItem->IsA<FNavigationToolAvaSequence>())
+	if (!InItem.AsModel()->IsA<FNavigationToolAvaSequence>())
 	{
 		return SNullWidget::NullWidget;
 	}
diff --git a/Engine/Plugins/Experimental/AvalancheSequenceNavigator/Source/AvalancheSequenceNavigator/Private/Columns/AvaNavigationToolStatusColumn.h b/Engine/Plugins/Experimental/AvalancheSequenceNavigator/Source/AvalancheSequenceNavigator/Private/Columns/AvaNavigationToolStatusColumn.h
--- a/Engine/Plugins/Experimental/AvalancheSequenceNavigator/Source/AvalancheSequenceNavigator/Private/Columns/AvaNavigationToolStatusColumn.h
+++ b/Engine/Plugins/Experimental/AvalancheSequenceNavigator/Source/AvalancheSequenceNavigator/Private/Columns/AvaNavigationToolStatusColumn.h
@@ -3,14 +3,19 @@
 #pragma once
 
 #include "Columns/NavigationToolColumn.h"
+#include "MVVM/ICastable.h"
+
+#define UE_API AVALANCHESEQUENCENAVIGATOR_API
 
 namespace UE::SequenceNavigator
 {
 
-class FAvaNavigationToolStatusColumn : public FNavigationToolColumn
+class FAvaNavigationToolStatusColumn
+	: public FNavigationToolColumn
 {
 public:
-	UE_NAVIGATIONTOOL_INHERITS_WITH_SUPER(FAvaNavigationToolStatusColumn, FNavigationToolColumn);
+	UE_SEQUENCER_DECLARE_CASTABLE_API(UE_API, FAvaNavigationToolStatusColumn
+		, FNavigationToolColumn);
 
 	static FName StaticColumnId() { return TEXT("Status"); }
 
@@ -22,10 +27,12 @@ protected:
 	virtual bool ShouldShowColumnByDefault() const override { return false; }
 	virtual float GetFillWidth() const override { return 10.f; }
 	virtual SHeaderRow::FColumn::FArguments ConstructHeaderRowColumn(const TSharedRef<INavigationToolView>& InToolView, const float InFillSize) override;
-	virtual TSharedRef<SWidget> ConstructRowWidget(const FNavigationToolItemRef& InItem
+	virtual TSharedRef<SWidget> ConstructRowWidget(const FNavigationToolViewModelPtr& InItem
 		, const TSharedRef<INavigationToolView>& InView
 		, const TSharedRef<SNavigationToolTreeRow>& InRow) override;
 	//~ End INavigationToolColumn
 };
 
 } // namespace UE::SequenceNavigator
+
+#undef UE_API
diff --git a/Engine/Plugins/Experimental/AvalancheSequenceNavigator/Source/AvalancheSequenceNavigator/Private/Columns/SAvaNavigationToolStatus.cpp b/Engine/Plugins/Experimental/AvalancheSequenceNavigator/Source/AvalancheSequenceNavigator/Private/Columns/SAvaNavigationToolStatus.cpp
--- a/Engine/Plugins/Experimental/AvalancheSequenceNavigator/Source/AvalancheSequenceNavigator/Private/Columns/SAvaNavigationToolStatus.cpp
+++ b/Engine/Plugins/Experimental/AvalancheSequenceNavigator/Source/AvalancheSequenceNavigator/Private/Columns/SAvaNavigationToolStatus.cpp
@@ -18,8 +18,10 @@
 namespace UE::SequenceNavigator
 {
 
+using namespace UE::Sequencer;
+
 void SAvaNavigationToolStatus::Construct(const FArguments& InArgs
-	, const FNavigationToolItemPtr& InItem
+	, const FNavigationToolViewModelPtr& InItem
 	, const TSharedRef<INavigationToolView>& InView
 	, const TSharedRef<SNavigationToolTreeRow>& InRowWidget
 	, const TSharedRef<ISequencer>& InSequencer)
@@ -29,8 +31,8 @@ void SAvaNavigationToolStatus::Construct(const FArguments& InArgs
 	WeakRowWidget = InRowWidget;
 	WeakSequencer = InSequencer;
 
-	FNavigationToolAvaSequence* const AvaSequenceItem = InItem->CastTo<FNavigationToolAvaSequence>();
-	if (!AvaSequenceItem)
+	const TViewModelPtr<FNavigationToolAvaSequence> AvaSequenceItem = InItem.ImplicitCast();
+	if (!AvaSequenceItem.IsValid())
 	{
 		return;
 	}
@@ -63,19 +65,20 @@ void SAvaNavigationToolStatus::Construct(const FArguments& InArgs
 
 void SAvaNavigationToolStatus::Tick(const FGeometry& InAllottedGeometry, const double InCurrentTime, const float InDeltaTime)
 {
-	if (!WeakItem.IsValid() || !WeakSequencer.IsValid())
+	const TSharedPtr<ISequencer> Sequencer = WeakSequencer.Pin();
+	if (!Sequencer.IsValid())
 	{
 		return;
 	}
 
-	const TSharedPtr<ISequencer> Sequencer = WeakSequencer.Pin();
-	if (!ensure(Sequencer.IsValid()))
+	const FNavigationToolViewModelPtr Item = WeakItem.Pin();
+	if (!Item.IsValid())
 	{
 		return;
 	}
 
-	const FNavigationToolAvaSequence* const AvaSequenceItem = WeakItem.Pin()->CastTo<FNavigationToolAvaSequence>();
-	if (!ensure(AvaSequenceItem))
+	const TViewModelPtr<FNavigationToolAvaSequence> AvaSequenceItem = Item.ImplicitCast();
+	if (!AvaSequenceItem.IsValid())
 	{
 		return;
 	}
@@ -87,7 +90,7 @@ void SAvaNavigationToolStatus::Tick(const FGeometry& InAllottedGeometry, const d
 	}
 
 	IAvaSceneInterface* const SceneInterface = FAvaSequencerUtils::GetSceneInterface(Sequencer.ToSharedRef());
-	if (!ensure(SceneInterface))
+	if (!SceneInterface)
 	{
 		return;
 	}
diff --git a/Engine/Plugins/Experimental/AvalancheSequenceNavigator/Source/AvalancheSequenceNavigator/Private/Columns/SAvaNavigationToolStatus.h b/Engine/Plugins/Experimental/AvalancheSequenceNavigator/Source/AvalancheSequenceNavigator/Private/Columns/SAvaNavigationToolStatus.h
--- a/Engine/Plugins/Experimental/AvalancheSequenceNavigator/Source/AvalancheSequenceNavigator/Private/Columns/SAvaNavigationToolStatus.h
+++ b/Engine/Plugins/Experimental/AvalancheSequenceNavigator/Source/AvalancheSequenceNavigator/Private/Columns/SAvaNavigationToolStatus.h
@@ -22,7 +22,7 @@ public:
 	SLATE_END_ARGS()
 
 	void Construct(const FArguments& InArgs
-		, const FNavigationToolItemPtr& InItem
+		, const FNavigationToolViewModelPtr& InItem
 		, const TSharedRef<INavigationToolView>& InView
 		, const TSharedRef<SNavigationToolTreeRow>& InRowWidget
 		, const TSharedRef<ISequencer>& InSequencer);
@@ -36,7 +36,7 @@ private:
 
 	TOptional<float> GetProgressPercent() const;
 
-	FNavigationToolItemWeakPtr WeakItem;
+	FNavigationToolViewModelWeakPtr WeakItem;
 	TWeakPtr<INavigationToolView> WeakView;
 	TWeakPtr<SNavigationToolTreeRow> WeakRowWidget;
 	TWeakPtr<ISequencer> WeakSequencer;
diff --git a/Engine/Plugins/Experimental/AvalancheSequenceNavigator/Source/AvalancheSequenceNavigator/Private/DragDrop/NavigationToolAvaSequenceDropHandler.cpp b/Engine/Plugins/Experimental/AvalancheSequenceNavigator/Source/AvalancheSequenceNavigator/Private/DragDrop/NavigationToolAvaSequenceDropHandler.cpp
--- a/Engine/Plugins/Experimental/AvalancheSequenceNavigator/Source/AvalancheSequenceNavigator/Private/DragDrop/NavigationToolAvaSequenceDropHandler.cpp
+++ b/Engine/Plugins/Experimental/AvalancheSequenceNavigator/Source/AvalancheSequenceNavigator/Private/DragDrop/NavigationToolAvaSequenceDropHandler.cpp
@@ -16,19 +16,23 @@
 
 namespace UE::SequenceNavigator
 {
+	
+using namespace UE::Sequencer;
+
+UE_SEQUENCER_DEFINE_CASTABLE(FNavigationToolAvaSequenceDropHandler);
 
 FNavigationToolAvaSequenceDropHandler::FNavigationToolAvaSequenceDropHandler(const TWeakPtr<IAvaSequencer>& InWeakAvaSequencer)
 	: WeakAvaSequencer(InWeakAvaSequencer)
 {
-	
 }
 
-bool FNavigationToolAvaSequenceDropHandler::IsDraggedItemSupported(const FNavigationToolItemPtr& InDraggedItem) const
+bool FNavigationToolAvaSequenceDropHandler::IsDraggedItemSupported(const FNavigationToolViewModelPtr& InDraggedItem) const
 {
-	return InDraggedItem->IsA<FNavigationToolAvaSequence>();
+	return InDraggedItem.AsModel()->IsA<FNavigationToolAvaSequence>();
 }
 
-TOptional<EItemDropZone> FNavigationToolAvaSequenceDropHandler::CanDrop(const EItemDropZone InDropZone, const FNavigationToolItemPtr& InTargetItem) const
+TOptional<EItemDropZone> FNavigationToolAvaSequenceDropHandler::CanDrop(const EItemDropZone InDropZone
+	, const FNavigationToolViewModelPtr& InTargetItem) const
 {
 	const bool bTargetIsRoot = (InTargetItem->GetItemId() == FNavigationToolItemId::RootId);
 
@@ -36,7 +40,7 @@ TOptional<EItemDropZone> FNavigationToolAvaSequenceDropHandler::CanDrop(const EI
 
 	if (!bTargetIsRoot)
 	{
-		if (FNavigationToolAvaSequence* const TargetAvaSequenceItem = InTargetItem->CastTo<FNavigationToolAvaSequence>())
+		if (FNavigationToolAvaSequence* const TargetAvaSequenceItem = InTargetItem.AsModel()->CastThis<FNavigationToolAvaSequence>())
 		{
 			TargetAvaSequence = TargetAvaSequenceItem->GetAvaSequence();
 		}
@@ -46,10 +50,16 @@ TOptional<EItemDropZone> FNavigationToolAvaSequenceDropHandler::CanDrop(const EI
 		}
 	}
 
-	for (const FNavigationToolItemPtr& Item : Items)
+	for (const FNavigationToolViewModelWeakPtr& WeakItem : WeakItems)
 	{
+		const FNavigationToolViewModelPtr Item = WeakItem.Pin();
+		if (!Item.IsValid())
+		{
+			continue;
+		}
+
 		// Only allow moving sequences to other Motion Design sequences
-		const FNavigationToolAvaSequence* const AvaSequenceItemToMove = Item->CastTo<FNavigationToolAvaSequence>();
+		const FNavigationToolAvaSequence* const AvaSequenceItemToMove = Item.AsModel()->CastThis<FNavigationToolAvaSequence>();
 		if (!AvaSequenceItemToMove)
 		{
 			return TOptional<EItemDropZone>();
@@ -75,7 +85,7 @@ TOptional<EItemDropZone> FNavigationToolAvaSequenceDropHandler::CanDrop(const EI
 	{
 		case ENavigationToolDragDropActionType::Move:
 			// Make sure the destination is not one of the items we're moving
-			if (!Items.Contains(InTargetItem))
+			if (!WeakItems.Contains(InTargetItem))
 			{
 				return InDropZone;
 			}
@@ -88,7 +98,7 @@ TOptional<EItemDropZone> FNavigationToolAvaSequenceDropHandler::CanDrop(const EI
 	return TOptional<EItemDropZone>();
 }
 
-bool FNavigationToolAvaSequenceDropHandler::Drop(const EItemDropZone InDropZone, const FNavigationToolItemPtr& InTargetItem)
+bool FNavigationToolAvaSequenceDropHandler::Drop(const EItemDropZone InDropZone, const FNavigationToolViewModelPtr& InTargetItem)
 {
 	switch (ActionType)
 	{
@@ -97,7 +107,7 @@ bool FNavigationToolAvaSequenceDropHandler::Drop(const EItemDropZone InDropZone,
 		break;
 
 	case ENavigationToolDragDropActionType::Copy:
-		DuplicateItems(Items, InTargetItem, InDropZone);
+		DuplicateItems(WeakItems, InTargetItem, InDropZone);
 		break;
 
 	default:
@@ -107,49 +117,60 @@ bool FNavigationToolAvaSequenceDropHandler::Drop(const EItemDropZone InDropZone,
 	return true;
 }
 
-void FNavigationToolAvaSequenceDropHandler::MoveItems(const EItemDropZone InDropZone, const FNavigationToolItemPtr& InTargetItem)
+void FNavigationToolAvaSequenceDropHandler::MoveItems(const EItemDropZone InDropZone, const FNavigationToolViewModelPtr& InTargetItem)
 {
-	TArray<TSharedPtr<INavigationToolAction>> ItemActions;
+	using namespace UE::Sequencer;
+
+	TArray<TSharedPtr<INavigationToolItemAction>> ItemActions;
 
 	FNavigationToolAddItemParams AddParams;
-	AddParams.RelativeItem = InTargetItem;
+	AddParams.WeakRelativeItem = InTargetItem;
 	AddParams.RelativeDropZone = InDropZone;
 	AddParams.Flags = ENavigationToolAddItemFlags::Select | ENavigationToolAddItemFlags::Transact;
 	AddParams.SelectionFlags = ENavigationToolItemSelectionFlags::AppendToCurrentSelection;
 
-	const TSet<FNavigationToolItemPtr> DraggedItemSet(Items);
+	const TSet<FNavigationToolViewModelWeakPtr> WeakDraggedItemSet(WeakItems);
 
 	// Remove all Items whose Parent is in the Item Set 
-	Items.RemoveAll([&DraggedItemSet](const FNavigationToolItemPtr& InItem)
+	WeakItems.RemoveAll([&WeakDraggedItemSet](const FNavigationToolViewModelWeakPtr& InWeakItem)
 		{
-			return !InItem.IsValid()
-				|| !InItem->GetParent().IsValid()
-				|| DraggedItemSet.Contains(InItem->GetParent());
+			if (const FNavigationToolViewModelPtr Item = InWeakItem.Pin())
+			{
+				const FNavigationToolViewModelPtr ItemParent = Item->GetParent();
+				return !ItemParent.IsValid()
+					|| WeakDraggedItemSet.Contains(ItemParent);
+			}
+			return false;
 		});
 
 	// Reverse order for onto since Item->AddChild(...) adds it to Index 0, so last item would be at the top, which is reversed
 	if (InDropZone == EItemDropZone::OntoItem)
 	{
-		Algo::Reverse(Items);
+		Algo::Reverse(WeakItems);
 	}
 
-	for (const FNavigationToolItemPtr& Item : Items)
+	for (const FNavigationToolViewModelWeakPtr& WeakItem : WeakItems)
 	{
-		// Only allow moving sequences to other Motion Design sequences
-		if (FNavigationToolAvaSequence* const AvaSequenceItemToMove = Item->CastTo<FNavigationToolAvaSequence>())
+		if (const FNavigationToolViewModelPtr Item = WeakItem.Pin())
 		{
-			AddParams.Item = SharedThis(AvaSequenceItemToMove);
-			ItemActions.Add(MakeShared<FNavigationToolAddItem>(AddParams));
+			// Only allow moving sequences to other Motion Design sequences
+			if (const TViewModelPtr<FNavigationToolAvaSequence> AvaSequenceItemToMove = Item.ImplicitCast())
+			{
+				AddParams.WeakItem = AvaSequenceItemToMove;
+				ItemActions.Add(MakeShared<FNavigationToolAddItem>(AddParams));
+			}
 		}
 	}
 
 	InTargetItem->GetOwnerTool().EnqueueItemActions(MoveTemp(ItemActions));
 }
 
-void FNavigationToolAvaSequenceDropHandler::DuplicateItems(const TArray<FNavigationToolItemPtr>& InItems
-	, const FNavigationToolItemPtr& InRelativeItem
+void FNavigationToolAvaSequenceDropHandler::DuplicateItems(const TArray<FNavigationToolViewModelWeakPtr>& InWeakItems
+	, const FNavigationToolViewModelWeakPtr& InWeakRelativeItem
 	, const TOptional<EItemDropZone>& InRelativeDropZone)
 {
+	using namespace UE::Sequencer;
+
 	const TSharedPtr<IAvaSequencer> AvaSequencer = WeakAvaSequencer.Pin();
 	if (!AvaSequencer.IsValid())
 	{
@@ -173,10 +194,16 @@ void FNavigationToolAvaSequenceDropHandler::DuplicateItems(const TArray<FNavigat
 	// Gather list of sequences to duplicate
 	TSet<UAvaSequence*> SequencesToDuplicate;
 
-	for (const FNavigationToolItemPtr& Item : InItems)
+	for (const FNavigationToolViewModelWeakPtr& WeakItem : InWeakItems)
 	{
-		const FNavigationToolAvaSequence* const AvaSequenceItem = Item->CastTo<FNavigationToolAvaSequence>();
-		if (!AvaSequenceItem)
+		const FNavigationToolViewModelPtr Item = WeakItem.Pin();
+		if (!Item.IsValid())
+		{
+			continue;
+		}
+
+		const TViewModelPtr<FNavigationToolAvaSequence> AvaSequenceItem = Item.ImplicitCast();
+		if (!AvaSequenceItem.IsValid())
 		{
 			continue;
 		}
diff --git a/Engine/Plugins/Experimental/AvalancheSequenceNavigator/Source/AvalancheSequenceNavigator/Private/DragDrop/NavigationToolAvaSequenceDropHandler.h b/Engine/Plugins/Experimental/AvalancheSequenceNavigator/Source/AvalancheSequenceNavigator/Private/DragDrop/NavigationToolAvaSequenceDropHandler.h
--- a/Engine/Plugins/Experimental/AvalancheSequenceNavigator/Source/AvalancheSequenceNavigator/Private/DragDrop/NavigationToolAvaSequenceDropHandler.h
+++ b/Engine/Plugins/Experimental/AvalancheSequenceNavigator/Source/AvalancheSequenceNavigator/Private/DragDrop/NavigationToolAvaSequenceDropHandler.h
@@ -5,6 +5,8 @@
 #include "Templates/SharedPointer.h"
 #include "DragDropOps/Handlers/NavigationToolItemDropHandler.h"
 
+#define UE_API AVALANCHESEQUENCENAVIGATOR_API
+
 class IAvaSequencer;
 enum class EItemDropZone;
 template<typename OptionalType> struct TOptional;
@@ -12,27 +14,32 @@ template<typename OptionalType> struct TOptional;
 namespace UE::SequenceNavigator
 {
 
-class FNavigationToolAvaSequenceDropHandler : public SequenceNavigator::FNavigationToolItemDropHandler
+class FNavigationToolAvaSequenceDropHandler
+	: public FNavigationToolItemDropHandler
 {
 public:
-	UE_NAVIGATIONTOOL_INHERITS(FNavigationToolAvaSequenceDropHandler, SequenceNavigator::FNavigationToolItemDropHandler);
+	UE_SEQUENCER_DECLARE_CASTABLE_API(UE_API, FNavigationToolAvaSequenceDropHandler
+		, FNavigationToolItemDropHandler);
 
 	FNavigationToolAvaSequenceDropHandler(const TWeakPtr<IAvaSequencer>& InWeakAvaSequencer);
 
 protected:
 	//~ Begin FNavigationToolItemDropHandler
-	virtual bool IsDraggedItemSupported(const SequenceNavigator::FNavigationToolItemPtr& InDraggedItem) const override;
-	virtual TOptional<EItemDropZone> CanDrop(const EItemDropZone InDropZone, const SequenceNavigator::FNavigationToolItemPtr& InTargetItem) const override;
-	virtual bool Drop(const EItemDropZone InDropZone, const SequenceNavigator::FNavigationToolItemPtr& InTargetItem) override;
+	virtual bool IsDraggedItemSupported(const FNavigationToolViewModelPtr& InDraggedItem) const override;
+	virtual TOptional<EItemDropZone> CanDrop(const EItemDropZone InDropZone
+		, const FNavigationToolViewModelPtr& InTargetItem) const override;
+	virtual bool Drop(const EItemDropZone InDropZone, const FNavigationToolViewModelPtr& InTargetItem) override;
 	//~ End FNavigationToolItemDropHandler
 
-	void MoveItems(const EItemDropZone InDropZone, const SequenceNavigator::FNavigationToolItemPtr& InTargetItem);
+	void MoveItems(const EItemDropZone InDropZone, const FNavigationToolViewModelPtr& InTargetItem);
 
-	void DuplicateItems(const TArray<SequenceNavigator::FNavigationToolItemPtr>& InItems
-		, const SequenceNavigator::FNavigationToolItemPtr& InRelativeItem
+	void DuplicateItems(const TArray<FNavigationToolViewModelWeakPtr>& InWeakItems
+		, const FNavigationToolViewModelWeakPtr& InWeakRelativeItem
 		, const TOptional<EItemDropZone>& InRelativeDropZone);
 
 	TWeakPtr<IAvaSequencer> WeakAvaSequencer;
 };
 
 } // namespace UE::SequenceNavigator
+
+#undef UE_API
diff --git a/Engine/Plugins/Experimental/AvalancheSequenceNavigator/Source/AvalancheSequenceNavigator/Private/Items/NavigationToolAvaSequence.cpp b/Engine/Plugins/Experimental/AvalancheSequenceNavigator/Source/AvalancheSequenceNavigator/Private/Items/NavigationToolAvaSequence.cpp
--- a/Engine/Plugins/Experimental/AvalancheSequenceNavigator/Source/AvalancheSequenceNavigator/Private/Items/NavigationToolAvaSequence.cpp
+++ b/Engine/Plugins/Experimental/AvalancheSequenceNavigator/Source/AvalancheSequenceNavigator/Private/Items/NavigationToolAvaSequence.cpp
@@ -19,8 +19,10 @@ namespace UE::SequenceNavigator
 
 using namespace UE::Sequencer;
 
+UE_SEQUENCER_DEFINE_CASTABLE(FNavigationToolAvaSequence)
+
 FNavigationToolAvaSequence::FNavigationToolAvaSequence(INavigationTool& InTool
-	, const FNavigationToolItemPtr& InParentItem
+	, const FNavigationToolViewModelPtr& InParentItem
 	, UAvaSequence* const InAvaSequence)
 	: FNavigationToolSequence(InTool
 	, InParentItem
@@ -32,13 +34,13 @@ FNavigationToolAvaSequence::FNavigationToolAvaSequence(INavigationTool& InTool
 
 bool FNavigationToolAvaSequence::AddChild(const FNavigationToolAddItemParams& InAddItemParams)
 {
-	if (!CanAddChild(InAddItemParams.Item))
+	const TViewModelPtr<FNavigationToolAvaSequence> AvaSequenceItemToAdd = InAddItemParams.WeakItem.ImplicitPin();
+	if (!AvaSequenceItemToAdd.IsValid())
 	{
 		return false;
 	}
 
-	const FNavigationToolAvaSequence* const AvaSequenceItemToAdd = InAddItemParams.Item->CastTo<FNavigationToolAvaSequence>();
-	if (!AvaSequenceItemToAdd)
+	if (!CanAddChild(AvaSequenceItemToAdd))
 	{
 		return false;
 	}
@@ -95,13 +97,19 @@ bool FNavigationToolAvaSequence::AddChild(const FNavigationToolAddItemParams& In
 
 bool FNavigationToolAvaSequence::RemoveChild(const FNavigationToolRemoveItemParams& InRemoveItemParams)
 {
+	const FNavigationToolViewModelPtr ItemToRemove = InRemoveItemParams.WeakItem.Pin();
+	if (!ItemToRemove.IsValid())
+	{
+		return false;
+	}
+
 	// Check that we're removing a child item that is directly under us, and not a grandchild
-	if (!InRemoveItemParams.Item.IsValid() || !Children.Contains(InRemoveItemParams.Item))
+	if (!WeakChildren.Contains(InRemoveItemParams.WeakItem))
 	{
 		return false;
 	}
 
-	if (const FNavigationToolAvaSequence* const AvaSequenceItemToRemove = InRemoveItemParams.Item->CastTo<FNavigationToolAvaSequence>())
+	if (const TViewModelPtr<FNavigationToolAvaSequence> AvaSequenceItemToRemove = ItemToRemove.ImplicitCast())
 	{
 		if (UAvaSequence* const AvaSequenceToRemove = AvaSequenceItemToRemove->GetAvaSequence())
 		{
@@ -113,12 +121,12 @@ bool FNavigationToolAvaSequence::RemoveChild(const FNavigationToolRemoveItemPara
 		}
 	}
 
-	return RemoveChildChecked(InRemoveItemParams.Item);
+	return RemoveChildChecked(ItemToRemove);
 }
 
-void FNavigationToolAvaSequence::FindChildren(TArray<FNavigationToolItemPtr>& OutChildren, bool bRecursive)
+void FNavigationToolAvaSequence::FindChildren(TArray<FNavigationToolViewModelWeakPtr>& OutWeakChildren, const bool bInRecursive)
 {
-	FNavigationToolSequence::FindChildren(OutChildren, bRecursive);
+	FNavigationToolSequence::FindChildren(OutWeakChildren, bInRecursive);
 
 	UAvaSequence* const AvaSequence = GetAvaSequence();
 	if (!AvaSequence)
@@ -148,16 +156,17 @@ void FNavigationToolAvaSequence::FindChildren(TArray<FNavigationToolItemPtr>& Ou
 			return InA->GetDisplayName().CompareTo(InB->GetDisplayName()) > 0;
 		});
 
-	const TSharedRef<FNavigationToolAvaSequence> SharedThisRef = SharedThis(this);
+	const TSharedRef<FNavigationToolProvider> ProviderRef = Provider.ToSharedRef();
+	const FNavigationToolViewModelPtr ParentItem = AsItemViewModel();
 
 	for (const TWeakObjectPtr<UAvaSequence>& SequenceWeak : ChildSequences)
 	{
-		const FNavigationToolItemPtr NewItem = Tool.FindOrAdd<FNavigationToolAvaSequence>(Provider
-			, SharedThisRef, SequenceWeak.Get());
-		OutChildren.Add(NewItem);
-		if (bRecursive)
+		const FNavigationToolViewModelPtr NewItem = Tool.FindOrAdd<FNavigationToolAvaSequence>(ProviderRef
+			, ParentItem, SequenceWeak.Get());
+		OutWeakChildren.Add(NewItem);
+		if (bInRecursive)
 		{
-			NewItem->FindChildren(OutChildren, bRecursive);
+			NewItem->FindChildren(OutWeakChildren, bInRecursive);
 		}
 	}
 }
@@ -172,29 +181,27 @@ bool FNavigationToolAvaSequence::CanRename() const
 	return GetAvaSequence() != nullptr;
 }
 
-bool FNavigationToolAvaSequence::Rename(const FString& InName)
+void FNavigationToolAvaSequence::Rename(const FText& InNewName)
 {
 	UAvaSequence* const UnderlyingSequence = GetAvaSequence();
 	if (!UnderlyingSequence)
 	{
-		return false;
+		return;
 	}
 
-	bool bRenamed = false;
+	const FString NewNameString = InNewName.ToString();
 
-	if (!InName.Equals(UnderlyingSequence->GetLabel().ToString(), ESearchCase::CaseSensitive))
+	if (NewNameString.Equals(UnderlyingSequence->GetLabel().ToString(), ESearchCase::CaseSensitive))
 	{
-		const FScopedTransaction Transaction(LOCTEXT("AvaSequenceRename", "Rename Sequence"));
-
-		UnderlyingSequence->Modify();
-		UnderlyingSequence->SetLabel(FName(InName));
+		return;
+	}
 
-		Tool.NotifyToolItemRenamed(SharedThis(this));
+	const FScopedTransaction Transaction(LOCTEXT("AvaSequenceRename", "Rename Sequence"));
 
-		bRenamed = true;
-	}
+	UnderlyingSequence->Modify();
+	UnderlyingSequence->SetLabel(FName(NewNameString));
 
-	return bRenamed;
+	Tool.NotifyToolItemRenamed(SharedThis(this));
 }
 
 bool FNavigationToolAvaSequence::CanDelete() const
@@ -265,7 +272,7 @@ FSlateColor FNavigationToolAvaSequence::GetIconColor() const
 
 void FNavigationToolAvaSequence::OnSelect()
 {
-	Super::OnSelect();
+	FNavigationToolSequence::OnSelect();
 }
 
 void FNavigationToolAvaSequence::OnDoubleClick()
diff --git a/Engine/Plugins/Experimental/AvalancheSequenceNavigator/Source/AvalancheSequenceNavigator/Private/Items/NavigationToolAvaSequence.h b/Engine/Plugins/Experimental/AvalancheSequenceNavigator/Source/AvalancheSequenceNavigator/Private/Items/NavigationToolAvaSequence.h
--- a/Engine/Plugins/Experimental/AvalancheSequenceNavigator/Source/AvalancheSequenceNavigator/Private/Items/NavigationToolAvaSequence.h
+++ b/Engine/Plugins/Experimental/AvalancheSequenceNavigator/Source/AvalancheSequenceNavigator/Private/Items/NavigationToolAvaSequence.h
@@ -15,24 +15,24 @@ namespace UE::SequenceNavigator
 class FNavigationToolAvaSequence : public FNavigationToolSequence
 {
 public:
-	UE_NAVIGATIONTOOL_INHERITS_WITH_SUPER(FNavigationToolAvaSequence, FNavigationToolSequence);
+	UE_SEQUENCER_DECLARE_VIEW_MODEL_TYPE_ID(FNavigationToolAvaSequence)
 
 	FNavigationToolAvaSequence(INavigationTool& InTool
-		, const FNavigationToolItemPtr& InParentItem
+		, const FNavigationToolViewModelPtr& InParentItem
 		, UAvaSequence* const InAvaSequence);
 
 	//~ Begin INavigationToolItem
 
 	virtual bool AddChild(const FNavigationToolAddItemParams& InAddItemParams) override;
 	virtual bool RemoveChild(const FNavigationToolRemoveItemParams& InRemoveItemParams) override;
-	virtual void FindChildren(TArray<FNavigationToolItemPtr>& OutChildren, bool bRecursive) override;
+	virtual void FindChildren(TArray<FNavigationToolViewModelWeakPtr>& OutWeakChildren, const bool bInRecursive) override;
 	virtual void GetItemProxies(TArray<TSharedPtr<FNavigationToolItemProxy>>& OutItemProxies) override;
 
 	virtual bool ShouldSort() const override { return true; }
 	virtual bool CanBeTopLevel() const override { return true; }
 
 	virtual bool CanRename() const override;
-	virtual bool Rename(const FString& InName) override;
+	virtual void Rename(const FText& InNewName) override;
 
 	virtual bool CanDelete() const override;
 	virtual bool Delete() override;
diff --git a/Engine/Plugins/Experimental/CinematicSequenceNavigator/Source/CinematicSequenceNavigator/CinematicSequenceNavigator.Build.cs b/Engine/Plugins/Experimental/CinematicSequenceNavigator/Source/CinematicSequenceNavigator/CinematicSequenceNavigator.Build.cs
--- a/Engine/Plugins/Experimental/CinematicSequenceNavigator/Source/CinematicSequenceNavigator/CinematicSequenceNavigator.Build.cs
+++ b/Engine/Plugins/Experimental/CinematicSequenceNavigator/Source/CinematicSequenceNavigator/CinematicSequenceNavigator.Build.cs
@@ -17,8 +17,9 @@ public class CinematicSequenceNavigator : ModuleRules
 			"CoreUObject",
 			"Engine",
 			"MovieScene",
-			"Sequencer",
 			"SequenceNavigator",
+			"Sequencer",
+			"SequencerCore",
 			"Slate",
 			"SlateCore",
 			"ToolMenus",
diff --git a/Engine/Plugins/Experimental/CinematicSequenceNavigator/Source/CinematicSequenceNavigator/Private/CinematicNavigationToolProvider.cpp b/Engine/Plugins/Experimental/CinematicSequenceNavigator/Source/CinematicSequenceNavigator/Private/CinematicNavigationToolProvider.cpp
--- a/Engine/Plugins/Experimental/CinematicSequenceNavigator/Source/CinematicSequenceNavigator/Private/CinematicNavigationToolProvider.cpp
+++ b/Engine/Plugins/Experimental/CinematicSequenceNavigator/Source/CinematicSequenceNavigator/Private/CinematicNavigationToolProvider.cpp
@@ -30,6 +30,7 @@
 #include "Items/NavigationToolItemParameters.h"
 #include "Items/NavigationToolSequence.h"
 #include "MovieScene.h"
+#include "NavigationToolDefines.h"
 #include "NavigationToolSettings.h"
 
 #define LOCTEXT_NAMESPACE "CinematicNavigationToolProvider"
@@ -177,14 +178,14 @@ void FCinematicNavigationToolProvider::OnExtendColumnViews(TSet<FNavigationToolC
 	FNavigationToolProvider::OnExtendColumnViews(OutColumnViews);
 }
 
-void FCinematicNavigationToolProvider::OnExtendItemChildren(UE::SequenceNavigator::INavigationTool& InTool
-	, const UE::SequenceNavigator::FNavigationToolItemPtr& InParentItem
-	, TArray<UE::SequenceNavigator::FNavigationToolItemPtr>& OutChildren
+void FCinematicNavigationToolProvider::OnExtendItemChildren(SequenceNavigator::INavigationTool& InTool
+	, const SequenceNavigator::FNavigationToolViewModelPtr& InParentItem
+	, TArray<SequenceNavigator::FNavigationToolViewModelWeakPtr>& OutWeakChildren
 	, const bool bInRecursive)
 {
 	using namespace UE::SequenceNavigator;
 
-	FNavigationToolProvider::OnExtendItemChildren(InTool, InParentItem, OutChildren, bInRecursive);
+	FNavigationToolProvider::OnExtendItemChildren(InTool, InParentItem, OutWeakChildren, bInRecursive);
 
 	// Only extending root item
 	if (InParentItem->GetItemId() != FNavigationToolItemId::RootId)
@@ -204,14 +205,15 @@ void FCinematicNavigationToolProvider::OnExtendItemChildren(UE::SequenceNavigato
 		return;
 	}
 
-	const TSharedRef<FCinematicNavigationToolProvider> Provider = SharedThis(this);
+	const TSharedRef<FCinematicNavigationToolProvider> ProviderRef = SharedThis(this);
 
-	const FNavigationToolItemPtr NewItem = InTool.FindOrAdd<FNavigationToolSequence>(Provider, InParentItem, RootSequence, nullptr, 0);
+	const FNavigationToolViewModelPtr NewItem = InTool.FindOrAdd<FNavigationToolSequence>(ProviderRef
+		, InParentItem, RootSequence, nullptr, 0);
 	const FNavigationToolItemFlagGuard Guard(NewItem, ENavigationToolItemFlags::IgnorePendingKill);
-	OutChildren.Add(NewItem);
+	OutWeakChildren.Add(NewItem);
 	if (bInRecursive)
 	{
-		NewItem->FindChildren(OutChildren, bInRecursive);
+		NewItem->FindChildren(OutWeakChildren, bInRecursive);
 	}
 }
 
diff --git a/Engine/Plugins/Experimental/CinematicSequenceNavigator/Source/CinematicSequenceNavigator/Private/CinematicNavigationToolProvider.h b/Engine/Plugins/Experimental/CinematicSequenceNavigator/Source/CinematicSequenceNavigator/Private/CinematicNavigationToolProvider.h
--- a/Engine/Plugins/Experimental/CinematicSequenceNavigator/Source/CinematicSequenceNavigator/Private/CinematicNavigationToolProvider.h
+++ b/Engine/Plugins/Experimental/CinematicSequenceNavigator/Source/CinematicSequenceNavigator/Private/CinematicNavigationToolProvider.h
@@ -54,8 +54,8 @@ public:
 	virtual void OnExtendColumns(SequenceNavigator::FNavigationToolColumnExtender& OutExtender) override;
 	virtual void OnExtendColumnViews(TSet<FNavigationToolColumnView>& OutColumnViews) override;
 	virtual void OnExtendItemChildren(SequenceNavigator::INavigationTool& InTool
-		, const SequenceNavigator::FNavigationToolItemPtr& InParentItem
-		, TArray<SequenceNavigator::FNavigationToolItemPtr>& OutChildren
+		, const SequenceNavigator::FNavigationToolViewModelPtr& InParentItem
+		, TArray<SequenceNavigator::FNavigationToolViewModelWeakPtr>& OutWeakChildren
 		, const bool bInRecursive) override;
 	virtual void OnExtendBuiltInFilters(TArray<FNavigationToolBuiltInFilterParams>& OutFilterParams) override;
 
diff --git a/Engine/Plugins/Experimental/LevelSequenceNavigatorBridge/Source/LevelSequenceNavigatorBridge/Private/LevelSequenceNavigationToolProvider.cpp b/Engine/Plugins/Experimental/LevelSequenceNavigatorBridge/Source/LevelSequenceNavigatorBridge/Private/LevelSequenceNavigationToolProvider.cpp
--- a/Engine/Plugins/Experimental/LevelSequenceNavigatorBridge/Source/LevelSequenceNavigatorBridge/Private/LevelSequenceNavigationToolProvider.cpp
+++ b/Engine/Plugins/Experimental/LevelSequenceNavigatorBridge/Source/LevelSequenceNavigatorBridge/Private/LevelSequenceNavigationToolProvider.cpp
@@ -25,6 +25,7 @@
 #include "LevelSequence.h"
 #include "MovieScene.h"
 #include "MovieSceneSequence.h"
+#include "NavigationToolDefines.h"
 #include "NavigationToolSettings.h"
 
 #define LOCTEXT_NAMESPACE "LevelSequenceNavigationToolProvider"
@@ -128,13 +129,13 @@ void FLevelSequenceNavigationToolProvider::OnExtendColumnViews(TSet<FNavigationT
 }
 
 void FLevelSequenceNavigationToolProvider::OnExtendItemChildren(UE::SequenceNavigator::INavigationTool& InTool
-	, const UE::SequenceNavigator::FNavigationToolItemPtr& InParentItem
-	, TArray<UE::SequenceNavigator::FNavigationToolItemPtr>& OutChildren
+	, const UE::SequenceNavigator::FNavigationToolViewModelPtr& InParentItem
+	, TArray<UE::SequenceNavigator::FNavigationToolViewModelWeakPtr>& OutWeakChildren
 	, const bool bInRecursive)
 {
 	using namespace UE::SequenceNavigator;
 
-	FNavigationToolProvider::OnExtendItemChildren(InTool, InParentItem, OutChildren, bInRecursive);
+	FNavigationToolProvider::OnExtendItemChildren(InTool, InParentItem, OutWeakChildren, bInRecursive);
 
 	// Only extending root item
 	if (InParentItem->GetItemId() != FNavigationToolItemId::RootId)
@@ -154,14 +155,15 @@ void FLevelSequenceNavigationToolProvider::OnExtendItemChildren(UE::SequenceNavi
 		return;
 	}
 
-	const TSharedRef<FLevelSequenceNavigationToolProvider> Provider = SharedThis(this);
+	const TSharedRef<FLevelSequenceNavigationToolProvider> ProviderRef = SharedThis(this);
 
-	const FNavigationToolItemPtr NewItem = InTool.FindOrAdd<FNavigationToolSequence>(Provider, InParentItem, RootSequence, nullptr, 0);
+	const FNavigationToolViewModelPtr NewItem = InTool.FindOrAdd<FNavigationToolSequence>(ProviderRef
+		, InParentItem, RootSequence, nullptr, 0);
 	const FNavigationToolItemFlagGuard Guard(NewItem, ENavigationToolItemFlags::IgnorePendingKill);
-	OutChildren.Add(NewItem);
+	OutWeakChildren.Add(NewItem);
 	if (bInRecursive)
 	{
-		NewItem->FindChildren(OutChildren, bInRecursive);
+		NewItem->FindChildren(OutWeakChildren, bInRecursive);
 	}
 }
 
diff --git a/Engine/Plugins/Experimental/LevelSequenceNavigatorBridge/Source/LevelSequenceNavigatorBridge/Private/LevelSequenceNavigationToolProvider.h b/Engine/Plugins/Experimental/LevelSequenceNavigatorBridge/Source/LevelSequenceNavigatorBridge/Private/LevelSequenceNavigationToolProvider.h
--- a/Engine/Plugins/Experimental/LevelSequenceNavigatorBridge/Source/LevelSequenceNavigatorBridge/Private/LevelSequenceNavigationToolProvider.h
+++ b/Engine/Plugins/Experimental/LevelSequenceNavigatorBridge/Source/LevelSequenceNavigatorBridge/Private/LevelSequenceNavigationToolProvider.h
@@ -47,8 +47,8 @@ public:
 	virtual void OnExtendColumns(UE::SequenceNavigator::FNavigationToolColumnExtender& OutExtender) override;
 	virtual void OnExtendColumnViews(TSet<FNavigationToolColumnView>& OutColumnViews) override;
 	virtual void OnExtendItemChildren(UE::SequenceNavigator::INavigationTool& InTool
-		, const UE::SequenceNavigator::FNavigationToolItemPtr& InParentItem
-		, TArray<UE::SequenceNavigator::FNavigationToolItemPtr>& OutChildren
+		, const UE::SequenceNavigator::FNavigationToolViewModelPtr& InParentItem
+		, TArray<UE::SequenceNavigator::FNavigationToolViewModelWeakPtr>& OutWeakChildren
 		, const bool bInRecursive) override;
 	virtual void OnExtendBuiltInFilters(TArray<FNavigationToolBuiltInFilterParams>& OutFilterParams) override;
 
diff --git a/Engine/Plugins/Experimental/SequenceNavigator/Source/SequenceNavigator/Private/Columns/NavigationToolColorColumn.cpp b/Engine/Plugins/Experimental/SequenceNavigator/Source/SequenceNavigator/Private/Columns/NavigationToolColorColumn.cpp
--- a/Engine/Plugins/Experimental/SequenceNavigator/Source/SequenceNavigator/Private/Columns/NavigationToolColorColumn.cpp
+++ b/Engine/Plugins/Experimental/SequenceNavigator/Source/SequenceNavigator/Private/Columns/NavigationToolColorColumn.cpp
@@ -2,6 +2,7 @@
 
 #include "Columns/NavigationToolColorColumn.h"
 #include "INavigationToolView.h"
+#include "Items/NavigationToolSequence.h"
 #include "Widgets/Columns/SNavigationToolColor.h"
 #include "Widgets/Layout/SBox.h"
 
@@ -10,6 +11,8 @@
 namespace UE::SequenceNavigator
 {
 
+UE_SEQUENCER_DEFINE_CASTABLE(FNavigationToolColorColumn)
+
 FText FNavigationToolColorColumn::GetColumnDisplayNameText() const
 {
 	return LOCTEXT("ColorColumn", "Color");
@@ -34,10 +37,15 @@ SHeaderRow::FColumn::FArguments FNavigationToolColorColumn::ConstructHeaderRowCo
 		];
 }
 
-TSharedRef<SWidget> FNavigationToolColorColumn::ConstructRowWidget(const FNavigationToolItemRef& InItem
+TSharedRef<SWidget> FNavigationToolColorColumn::ConstructRowWidget(const FNavigationToolViewModelPtr& InItem
 	, const TSharedRef<INavigationToolView>& InView
 	, const TSharedRef<SNavigationToolTreeRow>& InRow)
 {
+	if (!InItem.AsModel()->IsA<FNavigationToolSequence>())
+	{
+		return SNullWidget::NullWidget;
+	}
+
 	return SNew(SNavigationToolColor, InItem, InView, InRow);
 }
 
diff --git a/Engine/Plugins/Experimental/SequenceNavigator/Source/SequenceNavigator/Private/Columns/NavigationToolColumn.cpp b/Engine/Plugins/Experimental/SequenceNavigator/Source/SequenceNavigator/Private/Columns/NavigationToolColumn.cpp
--- a/Engine/Plugins/Experimental/SequenceNavigator/Source/SequenceNavigator/Private/Columns/NavigationToolColumn.cpp
+++ b/Engine/Plugins/Experimental/SequenceNavigator/Source/SequenceNavigator/Private/Columns/NavigationToolColumn.cpp
@@ -5,4 +5,6 @@
 namespace UE::SequenceNavigator
 {
 
+UE_SEQUENCER_DEFINE_CASTABLE(FNavigationToolColumn)
+
 }
diff --git a/Engine/Plugins/Experimental/SequenceNavigator/Source/SequenceNavigator/Private/Columns/NavigationToolCommentColumn.cpp b/Engine/Plugins/Experimental/SequenceNavigator/Source/SequenceNavigator/Private/Columns/NavigationToolCommentColumn.cpp
--- a/Engine/Plugins/Experimental/SequenceNavigator/Source/SequenceNavigator/Private/Columns/NavigationToolCommentColumn.cpp
+++ b/Engine/Plugins/Experimental/SequenceNavigator/Source/SequenceNavigator/Private/Columns/NavigationToolCommentColumn.cpp
@@ -9,6 +9,10 @@
 namespace UE::SequenceNavigator
 {
 
+using namespace Sequencer;
+
+UE_SEQUENCER_DEFINE_CASTABLE(FNavigationToolCommentColumn)
+
 FText FNavigationToolCommentColumn::GetColumnDisplayNameText() const
 {
 	return LOCTEXT("CommentColumn", "Comment");
@@ -34,11 +38,11 @@ SHeaderRow::FColumn::FArguments FNavigationToolCommentColumn::ConstructHeaderRow
 		.OnGetMenuContent(InView, &INavigationToolView::GetColumnMenuContent, ColumnId);
 }
 
-TSharedRef<SWidget> FNavigationToolCommentColumn::ConstructRowWidget(const FNavigationToolItemRef& InItem
+TSharedRef<SWidget> FNavigationToolCommentColumn::ConstructRowWidget(const FNavigationToolViewModelPtr& InItem
 	, const TSharedRef<INavigationToolView>& InView
 	, const TSharedRef<SNavigationToolTreeRow>& InRow)
 {
-	return SNew(SNavigationToolComment, InItem, InView, InRow);
+	return SNew(SNavigationToolComment, InItem.ImplicitCast(), InView, InRow);
 }
 
 } // namespace UE::SequenceNavigator
diff --git a/Engine/Plugins/Experimental/SequenceNavigator/Source/SequenceNavigator/Private/Columns/NavigationToolDeactiveStateColumn.cpp b/Engine/Plugins/Experimental/SequenceNavigator/Source/SequenceNavigator/Private/Columns/NavigationToolDeactiveStateColumn.cpp
--- a/Engine/Plugins/Experimental/SequenceNavigator/Source/SequenceNavigator/Private/Columns/NavigationToolDeactiveStateColumn.cpp
+++ b/Engine/Plugins/Experimental/SequenceNavigator/Source/SequenceNavigator/Private/Columns/NavigationToolDeactiveStateColumn.cpp
@@ -1,9 +1,10 @@
 // Copyright Epic Games, Inc. All Rights Reserved.
 
 #include "Columns/NavigationToolDeactiveStateColumn.h"
-#include "Extensions/ISequenceInactivableExtension.h"
 #include "INavigationToolView.h"
 #include "Items/INavigationToolItem.h"
+#include "Items/NavigationToolItem.h"
+#include "MVVM/Extensions/IDeactivatableExtension.h"
 #include "Widgets/Columns/SNavigationToolDeactiveState.h"
 #include "Widgets/SBoxPanel.h"
 
@@ -12,6 +13,8 @@
 namespace UE::SequenceNavigator
 {
 
+UE_SEQUENCER_DEFINE_CASTABLE(FNavigationToolDeactiveStateColumn)
+
 FText FNavigationToolDeactiveStateColumn::GetColumnDisplayNameText() const
 {
 	return LOCTEXT("NavigationToolDeactiveStateColumn", "Sequence Evaluation");
@@ -43,21 +46,18 @@ SHeaderRow::FColumn::FArguments FNavigationToolDeactiveStateColumn::ConstructHea
 		];
 }
 
-TSharedRef<SWidget> FNavigationToolDeactiveStateColumn::ConstructRowWidget(const FNavigationToolItemRef& InItem
+TSharedRef<SWidget> FNavigationToolDeactiveStateColumn::ConstructRowWidget(const FNavigationToolViewModelPtr& InItem
 	, const TSharedRef<INavigationToolView>& InView
 	, const TSharedRef<SNavigationToolTreeRow>& InRow)
 {
-	if (!InItem->IsA<ISequenceInactivableExtension>())
-	{
-		return SNullWidget::NullWidget;
-	}
+	using namespace Sequencer;
 
 	return SNew(SHorizontalBox)
 		+ SHorizontalBox::Slot()
 		.AutoWidth()
 		.VAlign(VAlign_Center)
 		[
-			SNew(SNavigationToolDeactiveState, SharedThis(this), InItem, InView, InRow)
+			SNew(SNavigationToolDeactiveState, SharedThis(this), InItem.ImplicitCast(), InView, InRow)
 		];
 }
 
diff --git a/Engine/Plugins/Experimental/SequenceNavigator/Source/SequenceNavigator/Private/Columns/NavigationToolHBiasColumn.cpp b/Engine/Plugins/Experimental/SequenceNavigator/Source/SequenceNavigator/Private/Columns/NavigationToolHBiasColumn.cpp
--- a/Engine/Plugins/Experimental/SequenceNavigator/Source/SequenceNavigator/Private/Columns/NavigationToolHBiasColumn.cpp
+++ b/Engine/Plugins/Experimental/SequenceNavigator/Source/SequenceNavigator/Private/Columns/NavigationToolHBiasColumn.cpp
@@ -11,6 +11,8 @@
 namespace UE::SequenceNavigator
 {
 
+UE_SEQUENCER_DEFINE_CASTABLE(FNavigationToolHBiasColumn)
+
 FText FNavigationToolHBiasColumn::GetColumnDisplayNameText() const
 {
 	return LOCTEXT("HBiasColumn", "HBias");
@@ -36,16 +38,11 @@ SHeaderRow::FColumn::FArguments FNavigationToolHBiasColumn::ConstructHeaderRowCo
 		.OnGetMenuContent(InView, &INavigationToolView::GetColumnMenuContent, ColumnId);
 }
 
-TSharedRef<SWidget> FNavigationToolHBiasColumn::ConstructRowWidget(const FNavigationToolItemRef& InItem
+TSharedRef<SWidget> FNavigationToolHBiasColumn::ConstructRowWidget(const FNavigationToolViewModelPtr& InItem
 	, const TSharedRef<INavigationToolView>& InView
 	, const TSharedRef<SNavigationToolTreeRow>& InRow)
 {
-	if (!InItem->IsA<FNavigationToolSequence>())
-	{
-		return SNullWidget::NullWidget;
-	}
-
-	return SNew(SNavigationToolHBias, InItem, InView, InRow);
+	return SNew(SNavigationToolHBias, InItem.ImplicitCast(), InView, InRow);
 }
 
 } // namespace UE::SequenceNavigator
diff --git a/Engine/Plugins/Experimental/SequenceNavigator/Source/SequenceNavigator/Private/Columns/NavigationToolInTimeColumn.cpp b/Engine/Plugins/Experimental/SequenceNavigator/Source/SequenceNavigator/Private/Columns/NavigationToolInTimeColumn.cpp
--- a/Engine/Plugins/Experimental/SequenceNavigator/Source/SequenceNavigator/Private/Columns/NavigationToolInTimeColumn.cpp
+++ b/Engine/Plugins/Experimental/SequenceNavigator/Source/SequenceNavigator/Private/Columns/NavigationToolInTimeColumn.cpp
@@ -3,7 +3,7 @@
 #include "Columns/NavigationToolInTimeColumn.h"
 #include "Extensions/IInTimeExtension.h"
 #include "INavigationToolView.h"
-#include "Items/INavigationToolItem.h"
+#include "Items/NavigationToolItem.h"
 #include "Widgets/Columns/SNavigationToolInTime.h"
 
 #define LOCTEXT_NAMESPACE "NavigationToolInTimeColumn"
@@ -11,6 +11,8 @@
 namespace UE::SequenceNavigator
 {
 
+UE_SEQUENCER_DEFINE_CASTABLE(FNavigationToolInTimeColumn)
+
 FText FNavigationToolInTimeColumn::GetColumnDisplayNameText() const
 {
 	return LOCTEXT("InTimeColumn", "In");
@@ -36,16 +38,11 @@ SHeaderRow::FColumn::FArguments FNavigationToolInTimeColumn::ConstructHeaderRowC
 		.OnGetMenuContent(InView, &INavigationToolView::GetColumnMenuContent, ColumnId);
 }
 
-TSharedRef<SWidget> FNavigationToolInTimeColumn::ConstructRowWidget(const FNavigationToolItemRef& InItem
+TSharedRef<SWidget> FNavigationToolInTimeColumn::ConstructRowWidget(const FNavigationToolViewModelPtr& InItem
 	, const TSharedRef<INavigationToolView>& InView
 	, const TSharedRef<SNavigationToolTreeRow>& InRow)
 {
-	if (!InItem->IsA<IInTimeExtension>())
-	{
-		return SNullWidget::NullWidget;
-	}
-
-	return SNew(SNavigationToolInTime, InItem, InView, InRow);
+	return SNew(SNavigationToolInTime, InItem.ImplicitCast(), InView, InRow);
 }
 
 } // namespace UE::SequenceNavigator
diff --git a/Engine/Plugins/Experimental/SequenceNavigator/Source/SequenceNavigator/Private/Columns/NavigationToolItemsColumn.cpp b/Engine/Plugins/Experimental/SequenceNavigator/Source/SequenceNavigator/Private/Columns/NavigationToolItemsColumn.cpp
--- a/Engine/Plugins/Experimental/SequenceNavigator/Source/SequenceNavigator/Private/Columns/NavigationToolItemsColumn.cpp
+++ b/Engine/Plugins/Experimental/SequenceNavigator/Source/SequenceNavigator/Private/Columns/NavigationToolItemsColumn.cpp
@@ -9,6 +9,8 @@
 namespace UE::SequenceNavigator
 {
 
+UE_SEQUENCER_DEFINE_CASTABLE(FNavigationToolItemsColumn)
+
 FText FNavigationToolItemsColumn::GetColumnDisplayNameText() const
 {
 	return LOCTEXT("DisplayName", "Items");
@@ -34,11 +36,11 @@ SHeaderRow::FColumn::FArguments FNavigationToolItemsColumn::ConstructHeaderRowCo
 		.OnGetMenuContent(InView, &INavigationToolView::GetColumnMenuContent, ColumnId);
 }
 
-TSharedRef<SWidget> FNavigationToolItemsColumn::ConstructRowWidget(const FNavigationToolItemRef& InItem
+TSharedRef<SWidget> FNavigationToolItemsColumn::ConstructRowWidget(const FNavigationToolViewModelPtr& InItem
 	, const TSharedRef<INavigationToolView>& InView
 	, const TSharedRef<SNavigationToolTreeRow>& InRow)
 {
-	return SNew(SNavigationToolItemList, InItem, InView, InRow);
+	return SNew(SNavigationToolItemList, InItem.ImplicitCast(), InView, InRow);
 }
 
 } // namespace UE::SequenceNavigator
diff --git a/Engine/Plugins/Experimental/SequenceNavigator/Source/SequenceNavigator/Private/Columns/NavigationToolLabelColumn.cpp b/Engine/Plugins/Experimental/SequenceNavigator/Source/SequenceNavigator/Private/Columns/NavigationToolLabelColumn.cpp
--- a/Engine/Plugins/Experimental/SequenceNavigator/Source/SequenceNavigator/Private/Columns/NavigationToolLabelColumn.cpp
+++ b/Engine/Plugins/Experimental/SequenceNavigator/Source/SequenceNavigator/Private/Columns/NavigationToolLabelColumn.cpp
@@ -2,13 +2,15 @@
 
 #include "Columns/NavigationToolLabelColumn.h"
 #include "INavigationToolView.h"
-#include "Items/INavigationToolItem.h"
+#include "Items/NavigationToolItem.h"
 
 #define LOCTEXT_NAMESPACE "NavigationToolLabelColumn"
 
 namespace UE::SequenceNavigator
 {
 
+UE_SEQUENCER_DEFINE_CASTABLE(FNavigationToolLabelColumn)
+
 FText FNavigationToolLabelColumn::GetColumnDisplayNameText() const
 {
 	return LOCTEXT("LabelColumn", "Label");
@@ -31,7 +33,7 @@ SHeaderRow::FColumn::FArguments FNavigationToolLabelColumn::ConstructHeaderRowCo
 		.OnGetMenuContent(InToolView, &INavigationToolView::GetColumnMenuContent, ColumnId);
 }
 
-TSharedRef<SWidget> FNavigationToolLabelColumn::ConstructRowWidget(const FNavigationToolItemRef& InItem
+TSharedRef<SWidget> FNavigationToolLabelColumn::ConstructRowWidget(const FNavigationToolViewModelPtr& InItem
 	, const TSharedRef<INavigationToolView>& InView
 	, const TSharedRef<SNavigationToolTreeRow>& InRow)
 {
diff --git a/Engine/Plugins/Experimental/SequenceNavigator/Source/SequenceNavigator/Private/Columns/NavigationToolLengthColumn.cpp b/Engine/Plugins/Experimental/SequenceNavigator/Source/SequenceNavigator/Private/Columns/NavigationToolLengthColumn.cpp
--- a/Engine/Plugins/Experimental/SequenceNavigator/Source/SequenceNavigator/Private/Columns/NavigationToolLengthColumn.cpp
+++ b/Engine/Plugins/Experimental/SequenceNavigator/Source/SequenceNavigator/Private/Columns/NavigationToolLengthColumn.cpp
@@ -9,6 +9,8 @@
 namespace UE::SequenceNavigator
 {
 
+UE_SEQUENCER_DEFINE_CASTABLE(FNavigationToolLengthColumn)
+
 FText FNavigationToolLengthColumn::GetColumnDisplayNameText() const
 {
 	return LOCTEXT("LengthColumn", "Length");
@@ -34,11 +36,11 @@ SHeaderRow::FColumn::FArguments FNavigationToolLengthColumn::ConstructHeaderRowC
 		.OnGetMenuContent(InView, &INavigationToolView::GetColumnMenuContent, ColumnId);
 }
 
-TSharedRef<SWidget> FNavigationToolLengthColumn::ConstructRowWidget(const FNavigationToolItemRef& InItem
+TSharedRef<SWidget> FNavigationToolLengthColumn::ConstructRowWidget(const FNavigationToolViewModelPtr& InItem
 	, const TSharedRef<INavigationToolView>& InView
 	, const TSharedRef<SNavigationToolTreeRow>& InRow)
 {
-	return SNew(SNavigationToolLength, InItem, InView, InRow);
+	return SNew(SNavigationToolLength, InItem.ImplicitCast(), InView, InRow);
 }
 
 } // namespace UE::SequenceNavigator
diff --git a/Engine/Plugins/Experimental/SequenceNavigator/Source/SequenceNavigator/Private/Columns/NavigationToolLockColumn.cpp b/Engine/Plugins/Experimental/SequenceNavigator/Source/SequenceNavigator/Private/Columns/NavigationToolLockColumn.cpp
--- a/Engine/Plugins/Experimental/SequenceNavigator/Source/SequenceNavigator/Private/Columns/NavigationToolLockColumn.cpp
+++ b/Engine/Plugins/Experimental/SequenceNavigator/Source/SequenceNavigator/Private/Columns/NavigationToolLockColumn.cpp
@@ -1,9 +1,9 @@
 // Copyright Epic Games, Inc. All Rights Reserved.
 
 #include "Columns/NavigationToolLockColumn.h"
-#include "Extensions/ISequenceLockableExtension.h"
 #include "INavigationToolView.h"
-#include "Items/INavigationToolItem.h"
+#include "Items/NavigationToolItem.h"
+#include "MVVM/Extensions/ILockableExtension.h"
 #include "Styling/StyleColors.h"
 #include "Widgets/Columns/SNavigationToolLock.h"
 #include "Widgets/SBoxPanel.h"
@@ -13,6 +13,8 @@
 namespace UE::SequenceNavigator
 {
 
+UE_SEQUENCER_DEFINE_CASTABLE(FNavigationToolLockColumn)
+
 FText FNavigationToolLockColumn::GetColumnDisplayNameText() const
 {
 	return LOCTEXT("LockColumn", "Lock");
@@ -43,21 +45,16 @@ SHeaderRow::FColumn::FArguments FNavigationToolLockColumn::ConstructHeaderRowCol
 		];
 }
 
-TSharedRef<SWidget> FNavigationToolLockColumn::ConstructRowWidget(const FNavigationToolItemRef& InItem
+TSharedRef<SWidget> FNavigationToolLockColumn::ConstructRowWidget(const FNavigationToolViewModelPtr& InItem
 	, const TSharedRef<INavigationToolView>& InView
 	, const TSharedRef<SNavigationToolTreeRow>& InRow)
 {
-	if (!InItem->IsA<ISequenceLockableExtension>())
-	{
-		return SNullWidget::NullWidget;
-	}
-
 	return SNew(SHorizontalBox)
 		+ SHorizontalBox::Slot()
 		.AutoWidth()
 		.VAlign(VAlign_Center)
 		[
-			SNew(SNavigationToolLock, InItem, InView, InRow)
+			SNew(SNavigationToolLock, InItem.ImplicitCast(), InView, InRow)
 		];
 }
 
diff --git a/Engine/Plugins/Experimental/SequenceNavigator/Source/SequenceNavigator/Private/Columns/NavigationToolMarkerVisibilityColumn.cpp b/Engine/Plugins/Experimental/SequenceNavigator/Source/SequenceNavigator/Private/Columns/NavigationToolMarkerVisibilityColumn.cpp
--- a/Engine/Plugins/Experimental/SequenceNavigator/Source/SequenceNavigator/Private/Columns/NavigationToolMarkerVisibilityColumn.cpp
+++ b/Engine/Plugins/Experimental/SequenceNavigator/Source/SequenceNavigator/Private/Columns/NavigationToolMarkerVisibilityColumn.cpp
@@ -2,8 +2,8 @@
 
 #include "Columns/NavigationToolMarkerVisibilityColumn.h"
 #include "Extensions/IMarkerVisibilityExtension.h"
-#include "Items/INavigationToolItem.h"
 #include "INavigationToolView.h"
+#include "Items/NavigationToolItem.h"
 #include "Widgets/Columns/SNavigationToolMarkerVisibility.h"
 #include "Widgets/SBoxPanel.h"
 #include "Widgets/SOverlay.h"
@@ -14,6 +14,8 @@
 namespace UE::SequenceNavigator
 {
 
+UE_SEQUENCER_DEFINE_CASTABLE(FNavigationToolMarkerVisibilityColumn)
+
 FText FNavigationToolMarkerVisibilityColumn::GetColumnDisplayNameText() const
 {
 	return LOCTEXT("NavigationToolMarkerVisibilityColumn", "Marker Visibility");
@@ -44,16 +46,11 @@ SHeaderRow::FColumn::FArguments FNavigationToolMarkerVisibilityColumn::Construct
 		];
 }
 
-TSharedRef<SWidget> FNavigationToolMarkerVisibilityColumn::ConstructRowWidget(const FNavigationToolItemRef& InItem
+TSharedRef<SWidget> FNavigationToolMarkerVisibilityColumn::ConstructRowWidget(const FNavigationToolViewModelPtr& InItem
 	, const TSharedRef<INavigationToolView>& InView
 	, const TSharedRef<SNavigationToolTreeRow>& InRow)
 {
-	if (!InItem->IsA<IMarkerVisibilityExtension>())
-	{
-		return SNullWidget::NullWidget;
-	}
-
-	return SNew(SNavigationToolMarkerVisibility, SharedThis(this), InItem, InView, InRow);
+	return SNew(SNavigationToolMarkerVisibility, SharedThis(this), InItem.ImplicitCast(), InView, InRow);
 }
 
 } // namespace UE::SequenceNavigator
diff --git a/Engine/Plugins/Experimental/SequenceNavigator/Source/SequenceNavigator/Private/Columns/NavigationToolOutTimeColumn.cpp b/Engine/Plugins/Experimental/SequenceNavigator/Source/SequenceNavigator/Private/Columns/NavigationToolOutTimeColumn.cpp
--- a/Engine/Plugins/Experimental/SequenceNavigator/Source/SequenceNavigator/Private/Columns/NavigationToolOutTimeColumn.cpp
+++ b/Engine/Plugins/Experimental/SequenceNavigator/Source/SequenceNavigator/Private/Columns/NavigationToolOutTimeColumn.cpp
@@ -3,7 +3,7 @@
 #include "Columns/NavigationToolOutTimeColumn.h"
 #include "Extensions/IOutTimeExtension.h"
 #include "INavigationToolView.h"
-#include "Items/INavigationToolItem.h"
+#include "Items/NavigationToolItem.h"
 #include "Widgets/Columns/SNavigationToolOutTime.h"
 #include "Widgets/SNavigationToolTreeRow.h"
 
@@ -12,6 +12,8 @@
 namespace UE::SequenceNavigator
 {
 
+UE_SEQUENCER_DEFINE_CASTABLE(FNavigationToolOutTimeColumn)
+
 FText FNavigationToolOutTimeColumn::GetColumnDisplayNameText() const
 {
 	return LOCTEXT("OutTimeColumn", "Out");
@@ -37,16 +39,11 @@ SHeaderRow::FColumn::FArguments FNavigationToolOutTimeColumn::ConstructHeaderRow
 		.OnGetMenuContent(InView, &INavigationToolView::GetColumnMenuContent, ColumnId);
 }
 
-TSharedRef<SWidget> FNavigationToolOutTimeColumn::ConstructRowWidget(const FNavigationToolItemRef& InItem
+TSharedRef<SWidget> FNavigationToolOutTimeColumn::ConstructRowWidget(const FNavigationToolViewModelPtr& InItem
 	, const TSharedRef<INavigationToolView>& InView
 	, const TSharedRef<SNavigationToolTreeRow>& InRow)
 {
-	if (!InItem->IsA<IOutTimeExtension>())
-	{
-		return SNullWidget::NullWidget;
-	}
-
-	return SNew(SNavigationToolOutTime, InItem, InView, InRow);
+	return SNew(SNavigationToolOutTime, InItem.ImplicitCast(), InView, InRow);
 }
 
 } // namespace UE::SequenceNavigator
diff --git a/Engine/Plugins/Experimental/SequenceNavigator/Source/SequenceNavigator/Private/Columns/NavigationToolPlayheadColumn.cpp b/Engine/Plugins/Experimental/SequenceNavigator/Source/SequenceNavigator/Private/Columns/NavigationToolPlayheadColumn.cpp
--- a/Engine/Plugins/Experimental/SequenceNavigator/Source/SequenceNavigator/Private/Columns/NavigationToolPlayheadColumn.cpp
+++ b/Engine/Plugins/Experimental/SequenceNavigator/Source/SequenceNavigator/Private/Columns/NavigationToolPlayheadColumn.cpp
@@ -3,7 +3,7 @@
 #include "Columns/NavigationToolPlayheadColumn.h"
 #include "Extensions/IPlayheadExtension.h"
 #include "INavigationToolView.h"
-#include "Items/INavigationToolItem.h"
+#include "Items/NavigationToolItem.h"
 #include "Widgets/Columns/SNavigationToolPlayhead.h"
 #include "Widgets/Views/SHeaderRow.h"
 
@@ -12,6 +12,8 @@
 namespace UE::SequenceNavigator
 {
 
+UE_SEQUENCER_DEFINE_CASTABLE(FNavigationToolPlayheadColumn)
+
 FText FNavigationToolPlayheadColumn::GetColumnDisplayNameText() const
 {
 	return LOCTEXT("NavigationToolPlayheadColumn", "Playhead");
@@ -42,16 +44,11 @@ SHeaderRow::FColumn::FArguments FNavigationToolPlayheadColumn::ConstructHeaderRo
 		];
 }
 
-TSharedRef<SWidget> FNavigationToolPlayheadColumn::ConstructRowWidget(const FNavigationToolItemRef& InItem
+TSharedRef<SWidget> FNavigationToolPlayheadColumn::ConstructRowWidget(const FNavigationToolViewModelPtr& InItem
 	, const TSharedRef<INavigationToolView>& InView
 	, const TSharedRef<SNavigationToolTreeRow>& InRow)
 {
-	if (!InItem->IsA<IPlayheadExtension>())
-	{
-		return SNullWidget::NullWidget;
-	}
-
-	return SNew(SNavigationToolPlayhead, SharedThis(this), InItem, InView, InRow);
+	return SNew(SNavigationToolPlayhead, SharedThis(this), InItem.ImplicitCast(), InView, InRow);
 }
 
 } // namespace UE::SequenceNavigator
diff --git a/Engine/Plugins/Experimental/SequenceNavigator/Source/SequenceNavigator/Private/Columns/NavigationToolRevisionControlColumn.cpp b/Engine/Plugins/Experimental/SequenceNavigator/Source/SequenceNavigator/Private/Columns/NavigationToolRevisionControlColumn.cpp
--- a/Engine/Plugins/Experimental/SequenceNavigator/Source/SequenceNavigator/Private/Columns/NavigationToolRevisionControlColumn.cpp
+++ b/Engine/Plugins/Experimental/SequenceNavigator/Source/SequenceNavigator/Private/Columns/NavigationToolRevisionControlColumn.cpp
@@ -13,6 +13,8 @@
 namespace UE::SequenceNavigator
 {
 
+UE_SEQUENCER_DEFINE_CASTABLE(FNavigationToolRevisionControlColumn)
+
 FText FNavigationToolRevisionControlColumn::GetColumnDisplayNameText() const
 {
 	return LOCTEXT("RevisionControl", "Revision Control");
@@ -44,16 +46,11 @@ SHeaderRow::FColumn::FArguments FNavigationToolRevisionControlColumn::ConstructH
 		];
 }
 
-TSharedRef<SWidget> FNavigationToolRevisionControlColumn::ConstructRowWidget(const FNavigationToolItemRef& InItem
+TSharedRef<SWidget> FNavigationToolRevisionControlColumn::ConstructRowWidget(const FNavigationToolViewModelPtr& InItem
 	, const TSharedRef<INavigationToolView>& InView
 	, const TSharedRef<SNavigationToolTreeRow>& InRow)
 {
-	if (!InItem->IsA<FNavigationToolSequence>())
-	{
-		return SNullWidget::NullWidget;
-	}
-
-	return SNew(SNavigationToolRevisionControl, SharedThis(this), InItem, InView, InRow);
+	return SNew(SNavigationToolRevisionControl, SharedThis(this), InItem.ImplicitCast(), InView, InRow);
 }
 
 } // namespace UE::SequenceNavigator
diff --git a/Engine/Plugins/Experimental/SequenceNavigator/Source/SequenceNavigator/Private/Columns/NavigationToolStartFrameOffsetColumn.cpp b/Engine/Plugins/Experimental/SequenceNavigator/Source/SequenceNavigator/Private/Columns/NavigationToolStartFrameOffsetColumn.cpp
--- a/Engine/Plugins/Experimental/SequenceNavigator/Source/SequenceNavigator/Private/Columns/NavigationToolStartFrameOffsetColumn.cpp
+++ b/Engine/Plugins/Experimental/SequenceNavigator/Source/SequenceNavigator/Private/Columns/NavigationToolStartFrameOffsetColumn.cpp
@@ -9,6 +9,8 @@
 namespace UE::SequenceNavigator
 {
 
+UE_SEQUENCER_DEFINE_CASTABLE(FNavigationToolStartFrameOffsetColumn)
+
 FText FNavigationToolStartFrameOffsetColumn::GetColumnDisplayNameText() const
 {
 	return LOCTEXT("StartFrameOffset", "Start Frame Offset");
@@ -34,11 +36,11 @@ SHeaderRow::FColumn::FArguments FNavigationToolStartFrameOffsetColumn::Construct
 		.OnGetMenuContent(InView, &INavigationToolView::GetColumnMenuContent, ColumnId);
 }
 
-TSharedRef<SWidget> FNavigationToolStartFrameOffsetColumn::ConstructRowWidget(const FNavigationToolItemRef& InItem
+TSharedRef<SWidget> FNavigationToolStartFrameOffsetColumn::ConstructRowWidget(const FNavigationToolViewModelPtr& InItem
 	, const TSharedRef<INavigationToolView>& InView
 	, const TSharedRef<SNavigationToolTreeRow>& InRow)
 {
-	return SNew(SNavigationToolStartFrameOffset, InItem, InView, InRow);
+	return SNew(SNavigationToolStartFrameOffset, InItem.ImplicitCast(), InView, InRow);
 }
 
 } // namespace UE::SequenceNavigator
diff --git a/Engine/Plugins/Experimental/SequenceNavigator/Source/SequenceNavigator/Private/Columns/NavigationToolTakeColumn.cpp b/Engine/Plugins/Experimental/SequenceNavigator/Source/SequenceNavigator/Private/Columns/NavigationToolTakeColumn.cpp
--- a/Engine/Plugins/Experimental/SequenceNavigator/Source/SequenceNavigator/Private/Columns/NavigationToolTakeColumn.cpp
+++ b/Engine/Plugins/Experimental/SequenceNavigator/Source/SequenceNavigator/Private/Columns/NavigationToolTakeColumn.cpp
@@ -11,6 +11,8 @@
 namespace UE::SequenceNavigator
 {
 
+UE_SEQUENCER_DEFINE_CASTABLE(FNavigationToolTakeColumn)
+
 FText FNavigationToolTakeColumn::GetColumnDisplayNameText() const
 {
 	return LOCTEXT("TakeColumn", "Take");
@@ -36,16 +38,11 @@ SHeaderRow::FColumn::FArguments FNavigationToolTakeColumn::ConstructHeaderRowCol
 		.OnGetMenuContent(InView, &INavigationToolView::GetColumnMenuContent, ColumnId);
 }
 
-TSharedRef<SWidget> FNavigationToolTakeColumn::ConstructRowWidget(const FNavigationToolItemRef& InItem
+TSharedRef<SWidget> FNavigationToolTakeColumn::ConstructRowWidget(const FNavigationToolViewModelPtr& InItem
 	, const TSharedRef<INavigationToolView>& InView
 	, const TSharedRef<SNavigationToolTreeRow>& InRow)
 {
-	if (!InItem->IsA<FNavigationToolSequence>())
-	{
-		return SNullWidget::NullWidget;
-	}
-
-	return SNew(SNavigationToolTake, InItem, InView, InRow);
+	return SNew(SNavigationToolTake, InItem.ImplicitCast(), InView, InRow);
 }
 
 } // namespace UE::SequenceNavigator
diff --git a/Engine/Plugins/Experimental/SequenceNavigator/Source/SequenceNavigator/Private/DragDropOps/Handlers/NavigationToolItemDropHandler.cpp b/Engine/Plugins/Experimental/SequenceNavigator/Source/SequenceNavigator/Private/DragDropOps/Handlers/NavigationToolItemDropHandler.cpp
--- a/Engine/Plugins/Experimental/SequenceNavigator/Source/SequenceNavigator/Private/DragDropOps/Handlers/NavigationToolItemDropHandler.cpp
+++ b/Engine/Plugins/Experimental/SequenceNavigator/Source/SequenceNavigator/Private/DragDropOps/Handlers/NavigationToolItemDropHandler.cpp
@@ -6,15 +6,20 @@
 namespace UE::SequenceNavigator
 {
 
+using namespace Sequencer;
+
+UE_SEQUENCER_DEFINE_CASTABLE(FNavigationToolItemDropHandler)
+
 void FNavigationToolItemDropHandler::Initialize(const FNavigationToolItemDragDropOp& InDragDropOp)
 {
-	Items = InDragDropOp.GetItems();
+	WeakItems = InDragDropOp.GetItems();
 	ActionType = InDragDropOp.GetActionType();
 
-	// Remove all Items that are Invalid or out of the Scope of this Handler
-	Items.RemoveAll([this](const FNavigationToolItemPtr& InItem)
+	// Remove all items that are invalid or out of the scope of this handler
+	WeakItems.RemoveAll([this](const FNavigationToolViewModelWeakPtr& InWeakItem)
 		{
-			return !InItem.IsValid() || !IsDraggedItemSupported(InItem);
+			const FNavigationToolViewModelPtr Item = InWeakItem.Pin();
+			return !Item.IsValid() || !IsDraggedItemSupported(Item);
 		});
 }
 
diff --git a/Engine/Plugins/Experimental/SequenceNavigator/Source/SequenceNavigator/Private/DragDropOps/Handlers/NavigationToolSequenceDropHandler.cpp b/Engine/Plugins/Experimental/SequenceNavigator/Source/SequenceNavigator/Private/DragDropOps/Handlers/NavigationToolSequenceDropHandler.cpp
--- a/Engine/Plugins/Experimental/SequenceNavigator/Source/SequenceNavigator/Private/DragDropOps/Handlers/NavigationToolSequenceDropHandler.cpp
+++ b/Engine/Plugins/Experimental/SequenceNavigator/Source/SequenceNavigator/Private/DragDropOps/Handlers/NavigationToolSequenceDropHandler.cpp
@@ -10,12 +10,16 @@
 namespace UE::SequenceNavigator
 {
 
-bool FNavigationToolSequenceDropHandler::IsDraggedItemSupported(const TSharedPtr<INavigationToolItem>& InDraggedItem) const
+using namespace Sequencer;
+
+UE_SEQUENCER_DEFINE_CASTABLE(FNavigationToolSequenceDropHandler)
+
+bool FNavigationToolSequenceDropHandler::IsDraggedItemSupported(const FNavigationToolViewModelPtr& InDraggedItem) const
 {
-	return InDraggedItem->IsA<FNavigationToolSequence>();
+	return InDraggedItem.AsModel()->IsA<FNavigationToolSequence>();
 }
 
-TOptional<EItemDropZone> FNavigationToolSequenceDropHandler::CanDrop(EItemDropZone InDropZone, const FNavigationToolItemPtr& InTargetItem) const
+TOptional<EItemDropZone> FNavigationToolSequenceDropHandler::CanDrop(EItemDropZone InDropZone, const FNavigationToolViewModelPtr& InTargetItem) const
 {
 	/*switch (ActionType)
 	{
@@ -34,7 +38,7 @@ TOptional<EItemDropZone> FNavigationToolSequenceDropHandler::CanDrop(EItemDropZo
 	return TOptional<EItemDropZone>();
 }
 
-bool FNavigationToolSequenceDropHandler::Drop(EItemDropZone InDropZone, const FNavigationToolItemPtr& InTargetItem)
+bool FNavigationToolSequenceDropHandler::Drop(EItemDropZone InDropZone, const FNavigationToolViewModelPtr& InTargetItem)
 {
 	// Disabling drag and drop for now
 	switch (ActionType)
@@ -54,7 +58,7 @@ bool FNavigationToolSequenceDropHandler::Drop(EItemDropZone InDropZone, const FN
 	return false;
 }
 
-void FNavigationToolSequenceDropHandler::MoveItems(EItemDropZone InDropZone, const FNavigationToolItemPtr& InTargetItem)
+void FNavigationToolSequenceDropHandler::MoveItems(EItemDropZone InDropZone, const FNavigationToolViewModelPtr& InTargetItem)
 {
 	// Can't access this mouse event at this point, so use the Slate Application to check modifiers.
 	/*const FAttachmentTransformRules& TransformRules = InTargetItem->GetOwnerTool()->GetProvider().GetTransformRule(!FSlateApplication::Get().GetModifierKeys().IsShiftDown());
@@ -66,10 +70,10 @@ void FNavigationToolSequenceDropHandler::MoveItems(EItemDropZone InDropZone, con
 	Params.AttachmentTransformRules = TransformRules;
 	Params.SelectionFlags           = ENavigationToolItemSelectionFlags::AppendToCurrentSelection;
 
-	const TSet<FNavigationToolItemPtr> DraggedItemSet(Items);
+	const TSet<FNavigationToolViewModelPtr> DraggedItemSet(Items);
 
 	// Remove all Items whose Parent is in the Item Set 
-	Items.RemoveAll([&DraggedItemSet](const FNavigationToolItemPtr& InItem)
+	Items.RemoveAll([&DraggedItemSet](const FNavigationToolViewModelPtr& InItem)
 		{
 			return !InItem.IsValid() || !InItem->GetParent().IsValid() || DraggedItemSet.Contains(InItem->GetParent());
 		});
@@ -82,7 +86,7 @@ void FNavigationToolSequenceDropHandler::MoveItems(EItemDropZone InDropZone, con
 
 	const TSharedRef<FNavigationTool> Tool = StaticCastSharedRef<FNavigationTool>(InTargetItem->GetOwnerTool());
 
-	for (const FNavigationToolItemPtr& Item : Items)
+	for (const FNavigationToolViewModelPtr& Item : Items)
 	{
 		Params.Item = Item;
 		Tool->EnqueueItemAction<FNavigationToolAddItem>(Params);
diff --git a/Engine/Plugins/Experimental/SequenceNavigator/Source/SequenceNavigator/Private/DragDropOps/Handlers/NavigationToolSequenceDropHandler.h b/Engine/Plugins/Experimental/SequenceNavigator/Source/SequenceNavigator/Private/DragDropOps/Handlers/NavigationToolSequenceDropHandler.h
--- a/Engine/Plugins/Experimental/SequenceNavigator/Source/SequenceNavigator/Private/DragDropOps/Handlers/NavigationToolSequenceDropHandler.h
+++ b/Engine/Plugins/Experimental/SequenceNavigator/Source/SequenceNavigator/Private/DragDropOps/Handlers/NavigationToolSequenceDropHandler.h
@@ -4,23 +4,29 @@
 
 #include "DragDropOps/Handlers/NavigationToolItemDropHandler.h"
 
+#define UE_API SEQUENCENAVIGATOR_API
+
 namespace UE::SequenceNavigator
 {
 
 /** Class that handles Dropping Sequence Items into a Target Item */
-class FNavigationToolSequenceDropHandler : public FNavigationToolItemDropHandler
+class FNavigationToolSequenceDropHandler
+	: public FNavigationToolItemDropHandler
 {
 public:
-	UE_NAVIGATIONTOOL_INHERITS(FNavigationToolSequenceDropHandler, FNavigationToolItemDropHandler);
+	UE_SEQUENCER_DECLARE_CASTABLE_API(UE_API, FNavigationToolSequenceDropHandler
+		, FNavigationToolItemDropHandler)
 
-private:
+protected:
 	//~ Begin FNavigationToolItemDropHandler
-	virtual bool IsDraggedItemSupported(const TSharedPtr<INavigationToolItem>& InDraggedItem) const override;
-	virtual TOptional<EItemDropZone> CanDrop(EItemDropZone InDropZone, const FNavigationToolItemPtr& InTargetItem) const override;
-	virtual bool Drop(EItemDropZone InDropZone, const FNavigationToolItemPtr& InTargetItem) override;
+	virtual bool IsDraggedItemSupported(const FNavigationToolViewModelPtr& InDraggedItem) const override;
+	virtual TOptional<EItemDropZone> CanDrop(EItemDropZone InDropZone, const FNavigationToolViewModelPtr& InTargetItem) const override;
+	virtual bool Drop(EItemDropZone InDropZone, const FNavigationToolViewModelPtr& InTargetItem) override;
 	//~ End FNavigationToolItemDropHandler
 
-	void MoveItems(EItemDropZone InDropZone, const FNavigationToolItemPtr& InTargetItem);
+	void MoveItems(EItemDropZone InDropZone, const FNavigationToolViewModelPtr& InTargetItem);
 };
 
 } // namespace UE::SequenceNavigator
+
+#undef UE_API
diff --git a/Engine/Plugins/Experimental/SequenceNavigator/Source/SequenceNavigator/Private/DragDropOps/NavigationToolItemDragDropOp.cpp b/Engine/Plugins/Experimental/SequenceNavigator/Source/SequenceNavigator/Private/DragDropOps/NavigationToolItemDragDropOp.cpp
--- a/Engine/Plugins/Experimental/SequenceNavigator/Source/SequenceNavigator/Private/DragDropOps/NavigationToolItemDragDropOp.cpp
+++ b/Engine/Plugins/Experimental/SequenceNavigator/Source/SequenceNavigator/Private/DragDropOps/NavigationToolItemDragDropOp.cpp
@@ -30,31 +30,36 @@ FText GetActionName(const ENavigationToolDragDropActionType InActionType)
 	return FText::GetEmpty();
 }
 
-FText GetItemName(const TArray<FNavigationToolItemPtr>& InItems)
+FText GetItemName(const TArray<FNavigationToolViewModelWeakPtr>& InWeakItems)
 {
-	if (InItems.IsEmpty())
+	using namespace Sequencer;
+
+	if (InWeakItems.IsEmpty())
 	{
 		return LOCTEXT("NoItems", "0 Items");
 	}
 
-	if (InItems.Num() == 1)
+	const FNavigationToolViewModelPtr FirstItem = InWeakItems[0].Pin();
+
+	if (InWeakItems.Num() == 1)
 	{
-		return InItems[0]->GetDisplayName();
+		return FirstItem->GetDisplayName();
 	}
 
-	return FText::Format(LOCTEXT("ManyItems", "{0} and {1} other item(s)"), InItems[0]->GetDisplayName(), InItems.Num() - 1);
+	return FText::Format(LOCTEXT("ManyItems", "{0} and {1} other item(s)")
+		, FirstItem->GetDisplayName(), InWeakItems.Num() - 1);
 }
 
-TSharedRef<FNavigationToolItemDragDropOp> FNavigationToolItemDragDropOp::New(const TArray<FNavigationToolItemPtr>& InItems
+TSharedRef<FNavigationToolItemDragDropOp> FNavigationToolItemDragDropOp::New(const TArray<FNavigationToolViewModelWeakPtr>& InWeakItems
 	, const TSharedPtr<FNavigationToolView>& InToolView
-	, ENavigationToolDragDropActionType InActionType)
+	, const ENavigationToolDragDropActionType InActionType)
 {
-	TSharedRef<FNavigationToolItemDragDropOp> DragDropOp = MakeShared<FNavigationToolItemDragDropOp>();
-	DragDropOp->Init(InItems, InToolView, InActionType);
+	const TSharedRef<FNavigationToolItemDragDropOp> DragDropOp = MakeShared<FNavigationToolItemDragDropOp>();
+	DragDropOp->Init(InWeakItems, InToolView, InActionType);
 	return DragDropOp;
 }
 
-FReply FNavigationToolItemDragDropOp::Drop(EItemDropZone InDropZone, const FNavigationToolItemPtr& InTargetItem)
+FReply FNavigationToolItemDragDropOp::Drop(EItemDropZone InDropZone, const FNavigationToolViewModelPtr& InTargetItem)
 {
 	FReply Reply = FReply::Unhandled();
 
@@ -81,7 +86,7 @@ FReply FNavigationToolItemDragDropOp::Drop(EItemDropZone InDropZone, const FNavi
 	return Reply;
 }
 
-TOptional<EItemDropZone> FNavigationToolItemDragDropOp::CanDrop(const EItemDropZone InDropZone, const FNavigationToolItemPtr& InTargetItem) const
+TOptional<EItemDropZone> FNavigationToolItemDragDropOp::CanDrop(const EItemDropZone InDropZone, const FNavigationToolViewModelPtr& InTargetItem) const
 {
 	const TSharedPtr<INavigationToolView> ToolView = GetToolView();
 
@@ -109,29 +114,29 @@ TOptional<EItemDropZone> FNavigationToolItemDragDropOp::CanDrop(const EItemDropZ
 	return TOptional<EItemDropZone>();
 }
 
-void FNavigationToolItemDragDropOp::Init(const TArray<FNavigationToolItemPtr>& InItems
+void FNavigationToolItemDragDropOp::Init(const TArray<FNavigationToolViewModelWeakPtr>& InWeakItems
 	, const TSharedPtr<FNavigationToolView>& InToolView
-	, ENavigationToolDragDropActionType InActionType)
+	, const ENavigationToolDragDropActionType InActionType)
 {
-	Items = InItems;
+	WeakItems = InWeakItems;
 	WeakToolView = InToolView;
 	ActionType = InActionType;
 	MouseCursor = EMouseCursor::GrabHandClosed;
 
-	CurrentIconBrush = Items.Num() == 1
-		? Items[0]->GetIconBrush()
+	CurrentIconBrush = WeakItems.Num() == 1
+		? WeakItems[0].Pin()->GetIconBrush()
 		: FSlateIconFinder::FindIconForClass(ALevelSequenceActor::StaticClass()).GetIcon();
 
 	CurrentHoverText = FText::Format(LOCTEXT("HoverText", "{0} {1}")
-		, GetActionName(InActionType)
-		, GetItemName(InItems));
+		, GetActionName(ActionType)
+		, GetItemName(WeakItems));
 
 	CurrentIconColorAndOpacity = FSlateColor::UseForeground();
 
 	SetupDefaults();
 	Construct();
 
-	// Add Default Drop Handlers 
+	// Add default drop handlers 
 	AddDropHandler<FNavigationToolSequenceDropHandler>();
 
 	FNavigationToolExtender::OnItemDragDropOpInitialized().Broadcast(*this);
diff --git a/Engine/Plugins/Experimental/SequenceNavigator/Source/SequenceNavigator/Private/Filters/Filters/NavigationToolBuiltInFilter.cpp b/Engine/Plugins/Experimental/SequenceNavigator/Source/SequenceNavigator/Private/Filters/Filters/NavigationToolBuiltInFilter.cpp
--- a/Engine/Plugins/Experimental/SequenceNavigator/Source/SequenceNavigator/Private/Filters/Filters/NavigationToolBuiltInFilter.cpp
+++ b/Engine/Plugins/Experimental/SequenceNavigator/Source/SequenceNavigator/Private/Filters/Filters/NavigationToolBuiltInFilter.cpp
@@ -7,8 +7,10 @@
 namespace UE::SequenceNavigator
 {
 
+using namespace Sequencer;
+
 FNavigationToolBuiltInFilter::FNavigationToolBuiltInFilter(const FNavigationToolBuiltInFilterParams& InFilterParams)
-	: FFilterBase<FNavigationToolItemPtr>(nullptr)
+	: FFilterBase<FNavigationToolViewModelPtr>(nullptr)
 	, bIsActive(true)
 	, FilterParams(InFilterParams)
 {
@@ -53,17 +55,17 @@ TSharedPtr<FUICommandInfo> FNavigationToolBuiltInFilter::GetToggleCommand() cons
 	return FilterParams.GetToggleCommand();
 }
 
-bool FNavigationToolBuiltInFilter::PassesFilter(const FNavigationToolFilterType InItem) const
+bool FNavigationToolBuiltInFilter::PassesFilter(const FNavigationToolViewModelPtr InItem) const
 {
-	auto ItemObjectPassesFilter = [this](const FNavigationToolItemPtr& InItem)
+	auto ItemObjectPassesFilter = [this](const FNavigationToolViewModelPtr& InItem)
 		{
 			UObject* const ItemObject = InItem->GetItemObject();
 			return ItemObject && FilterParams.IsValidObjectClass(ItemObject->GetClass());
 		};
 
-	auto ItemTypePassesFilter = [this](const FNavigationToolItemPtr& InItem)
+	auto ItemTypePassesFilter = [this](const FNavigationToolViewModelPtr& InItem)
 		{
-			return FilterParams.IsValidItemClass(InItem->GetTypeId());
+			return FilterParams.IsValidItemClass(InItem->ID);
 		};
 
 	const ENavigationToolFilterMode FilterMode = FilterParams.GetFilterMode();
@@ -78,15 +80,15 @@ bool FNavigationToolBuiltInFilter::PassesFilter(const FNavigationToolFilterType
 
 	if (EnumHasAnyFlags(FilterMode, ENavigationToolFilterMode::ContainerOfType))
 	{
-		TArray<FNavigationToolItemPtr> RemainingItems = InItem->GetChildren();
+		TArray<FNavigationToolViewModelWeakPtr> WeakRemainingItems = InItem->GetChildren();
 
-		while (!RemainingItems.IsEmpty())
+		while (!WeakRemainingItems.IsEmpty())
 		{
-			FNavigationToolItemPtr Item = RemainingItems.Pop();
+			const FNavigationToolViewModelPtr Item = WeakRemainingItems.Pop().Pin();
 
 			// Stop here if Item is invalid or if it's a Sequence Item as any children of this item will be considered contained by the item itself, not the querying item
 			// this could be improved if the Item Type filter as a whole ever moves past only checking UObject types, and moving past that Sequences are the top level items
-			if (!Item.IsValid() || Item->IsA<FNavigationToolSequence>())
+			if (!Item.IsValid() || Item.AsModel()->IsA<FNavigationToolSequence>())
 			{
 				continue;
 			}
@@ -96,7 +98,7 @@ bool FNavigationToolBuiltInFilter::PassesFilter(const FNavigationToolFilterType
 				return true;
 			}
 
-			RemainingItems.Append(Item->GetChildren());
+			WeakRemainingItems.Append(Item->GetChildren());
 		}
 	}
 
@@ -110,7 +112,7 @@ bool FNavigationToolBuiltInFilter::IsActive() const
 
 void FNavigationToolBuiltInFilter::SetActive(const bool bInActive)
 {
-	FFilterBase<FNavigationToolItemPtr>::SetActive(bInActive);
+	FFilterBase<FNavigationToolViewModelPtr>::SetActive(bInActive);
 	
 	bIsActive = bInActive;
 }
diff --git a/Engine/Plugins/Experimental/SequenceNavigator/Source/SequenceNavigator/Private/Filters/Filters/NavigationToolBuiltInFilter.h b/Engine/Plugins/Experimental/SequenceNavigator/Source/SequenceNavigator/Private/Filters/Filters/NavigationToolBuiltInFilter.h
--- a/Engine/Plugins/Experimental/SequenceNavigator/Source/SequenceNavigator/Private/Filters/Filters/NavigationToolBuiltInFilter.h
+++ b/Engine/Plugins/Experimental/SequenceNavigator/Source/SequenceNavigator/Private/Filters/Filters/NavigationToolBuiltInFilter.h
@@ -9,7 +9,7 @@
 namespace UE::SequenceNavigator
 {
 
-class FNavigationToolBuiltInFilter : public FFilterBase<FNavigationToolItemPtr>
+class FNavigationToolBuiltInFilter : public FFilterBase<FNavigationToolViewModelPtr>
 {
 public:
 	FNavigationToolBuiltInFilter(const FNavigationToolBuiltInFilterParams& InFilterParams);
@@ -24,7 +24,7 @@ public:
 
 	//~ Begin IFilter
 	virtual FString GetName() const override;
-	virtual bool PassesFilter(const FNavigationToolFilterType InItem) const override;
+	virtual bool PassesFilter(const FNavigationToolViewModelPtr InItem) const override;
 	//~ End IFilter
 
 	virtual FSlateIcon GetIcon() const;
diff --git a/Engine/Plugins/Experimental/SequenceNavigator/Source/SequenceNavigator/Private/Filters/Filters/NavigationToolBuiltInFilterParams.cpp b/Engine/Plugins/Experimental/SequenceNavigator/Source/SequenceNavigator/Private/Filters/Filters/NavigationToolBuiltInFilterParams.cpp
--- a/Engine/Plugins/Experimental/SequenceNavigator/Source/SequenceNavigator/Private/Filters/Filters/NavigationToolBuiltInFilterParams.cpp
+++ b/Engine/Plugins/Experimental/SequenceNavigator/Source/SequenceNavigator/Private/Filters/Filters/NavigationToolBuiltInFilterParams.cpp
@@ -11,7 +11,7 @@
 #include "Items/NavigationToolTrack.h"
 #include "LevelSequence.h"
 #include "LevelSequenceActor.h"
-#include "NavigationToolItemTypeId.h"
+#include "MVVM/ViewModelTypeID.h"
 #include "Styling/SlateIconFinder.h"
 
 #include UE_INLINE_GENERATED_CPP_BY_NAME(NavigationToolBuiltInFilterParams)
@@ -19,7 +19,7 @@
 #define LOCTEXT_NAMESPACE "NavigationToolBuiltInFilterParams"
 
 FNavigationToolBuiltInFilterParams::FNavigationToolBuiltInFilterParams(const FName InFilterId
-	, const TSet<UE::SequenceNavigator::FNavigationToolItemTypeId>& InItemClasses
+	, const TSet<UE::Sequencer::FViewModelTypeID>& InItemClasses
 	, const TArray<TSubclassOf<UObject>>& InFilterClasses
 	, const ENavigationToolFilterMode InFilterMode
 	, const FSlateBrush* const InIconBrush
@@ -128,13 +128,9 @@ void FNavigationToolBuiltInFilterParams::SetOverrideIconColor(const FSlateColor&
 	OverrideIcon.TintColor = InNewIconColor;
 }
 
-bool FNavigationToolBuiltInFilterParams::IsValidItemClass(const UE::SequenceNavigator::FNavigationToolItemTypeId InClassTypeId) const
+bool FNavigationToolBuiltInFilterParams::IsValidItemClass(const UE::Sequencer::FViewModelTypeID& InClassTypeId) const
 {
-	if (InClassTypeId.IsValid())
-	{
-		return ItemClasses.IsEmpty() || ItemClasses.Contains(InClassTypeId);
-	}
-	return false;
+	return ItemClasses.IsEmpty() || ItemClasses.Contains(InClassTypeId);
 }
 
 bool FNavigationToolBuiltInFilterParams::IsValidObjectClass(const UClass* const InClass) const
@@ -163,7 +159,7 @@ void FNavigationToolBuiltInFilterParams::SetFilterText(const FText& InText)
 	FilterText = InText;
 }
 
-bool FNavigationToolBuiltInFilterParams::PassesFilterText(const UE::SequenceNavigator::FNavigationToolItemPtr& InItem) const
+bool FNavigationToolBuiltInFilterParams::PassesFilterText(const UE::SequenceNavigator::FNavigationToolViewModelPtr& InItem) const
 {
 	// Since it needs to be a const function we have to create the TextFilter here instead of being a variable
 	if (FilterText.IsEmptyOrWhitespace())
@@ -181,7 +177,7 @@ bool FNavigationToolBuiltInFilterParams::PassesFilterText(const UE::SequenceNavi
 FNavigationToolBuiltInFilterParams FNavigationToolBuiltInFilterParams::CreateSequenceFilter()
 {
 	return FNavigationToolBuiltInFilterParams(TEXT("Level Sequence")
-		, { UE::SequenceNavigator::FNavigationToolSequence::GetStaticTypeId() }
+		, { UE::SequenceNavigator::FNavigationToolSequence::ID }
 		, { ULevelSequence::StaticClass(), UMovieSceneSequence::StaticClass() }
 		, ENavigationToolFilterMode::MatchesType
 		, FSlateIconFinder::FindIconForClass(ALevelSequenceActor::StaticClass()).GetIcon()
@@ -193,7 +189,7 @@ FNavigationToolBuiltInFilterParams FNavigationToolBuiltInFilterParams::CreateSeq
 FNavigationToolBuiltInFilterParams FNavigationToolBuiltInFilterParams::CreateTrackFilter()
 {
 	return FNavigationToolBuiltInFilterParams(TEXT("Track")
-		, { UE::SequenceNavigator::FNavigationToolTrack::GetStaticTypeId() }
+		, { UE::SequenceNavigator::FNavigationToolTrack::ID }
 		, { UMovieSceneTrack::StaticClass() }
 		, ENavigationToolFilterMode::MatchesType
 		, FSlateIconFinder::FindIconForClass(UMovieSceneTrack::StaticClass()).GetIcon()
@@ -205,9 +201,9 @@ FNavigationToolBuiltInFilterParams FNavigationToolBuiltInFilterParams::CreateTra
 FNavigationToolBuiltInFilterParams FNavigationToolBuiltInFilterParams::CreateBindingFilter()
 {
 	return FNavigationToolBuiltInFilterParams(TEXT("Binding")
-		, { UE::SequenceNavigator::FNavigationToolBinding::GetStaticTypeId()
-			, UE::SequenceNavigator::FNavigationToolActor::GetStaticTypeId()
-			, UE::SequenceNavigator::FNavigationToolComponent::GetStaticTypeId() }
+		, { UE::SequenceNavigator::FNavigationToolBinding::ID
+			, UE::SequenceNavigator::FNavigationToolActor::ID
+			, UE::SequenceNavigator::FNavigationToolComponent::ID }
 		, { }
 		, ENavigationToolFilterMode::MatchesType
 		, nullptr
@@ -219,7 +215,7 @@ FNavigationToolBuiltInFilterParams FNavigationToolBuiltInFilterParams::CreateBin
 FNavigationToolBuiltInFilterParams FNavigationToolBuiltInFilterParams::CreateMarkerFilter()
 {
 	return FNavigationToolBuiltInFilterParams(TEXT("Marker")
-		, { UE::SequenceNavigator::FNavigationToolMarker::GetStaticTypeId() }
+		, { UE::SequenceNavigator::FNavigationToolMarker::ID }
 		, { }
 		, ENavigationToolFilterMode::MatchesType
 		, nullptr
diff --git a/Engine/Plugins/Experimental/SequenceNavigator/Source/SequenceNavigator/Private/Filters/Filters/NavigationToolFilterBase.cpp b/Engine/Plugins/Experimental/SequenceNavigator/Source/SequenceNavigator/Private/Filters/Filters/NavigationToolFilterBase.cpp
--- a/Engine/Plugins/Experimental/SequenceNavigator/Source/SequenceNavigator/Private/Filters/Filters/NavigationToolFilterBase.cpp
+++ b/Engine/Plugins/Experimental/SequenceNavigator/Source/SequenceNavigator/Private/Filters/Filters/NavigationToolFilterBase.cpp
@@ -9,10 +9,10 @@
 namespace UE::SequenceNavigator
 {
 
-using namespace UE::Sequencer;
+using namespace Sequencer;
 
 FNavigationToolFilter::FNavigationToolFilter(INavigationToolFilterBar& InOutFilterInterface, TSharedPtr<FFilterCategory>&& InCategory)
-	: FSequencerFilterBase<FNavigationToolItemPtr>(InOutFilterInterface, MoveTemp(InCategory))
+	: FSequencerFilterBase<FNavigationToolViewModelPtr>(InOutFilterInterface, MoveTemp(InCategory))
 {
 }
 
diff --git a/Engine/Plugins/Experimental/SequenceNavigator/Source/SequenceNavigator/Private/Filters/Filters/NavigationToolFilterCollection.cpp b/Engine/Plugins/Experimental/SequenceNavigator/Source/SequenceNavigator/Private/Filters/Filters/NavigationToolFilterCollection.cpp
--- a/Engine/Plugins/Experimental/SequenceNavigator/Source/SequenceNavigator/Private/Filters/Filters/NavigationToolFilterCollection.cpp
+++ b/Engine/Plugins/Experimental/SequenceNavigator/Source/SequenceNavigator/Private/Filters/Filters/NavigationToolFilterCollection.cpp
@@ -6,7 +6,7 @@
 namespace UE::SequenceNavigator
 {
 
-using namespace UE::Sequencer;
+using namespace Sequencer;
 
 FNavigationToolFilterCollection::FNavigationToolFilterCollection(ISequencerFilterBar& InFilterInterface)
 	: FilterInterface(InFilterInterface)
@@ -33,9 +33,9 @@ bool FNavigationToolFilterCollection::ContainsFilter(const TSharedRef<FNavigatio
 
 void FNavigationToolFilterCollection::RemoveAll()
 {
-	using namespace UE::Sequencer;
+	using namespace Sequencer;
 
-	for (const TSharedPtr<IFilter<FNavigationToolItemPtr>>& Filter : ChildFilters)
+	for (const TSharedPtr<IFilter<FNavigationToolViewModelPtr>>& Filter : ChildFilters)
 	{
 		Filter->OnChanged().RemoveAll(this);
 	}
@@ -96,7 +96,7 @@ bool FNavigationToolFilterCollection::IsEmpty() const
 
 void FNavigationToolFilterCollection::Sort()
 {
-	ChildFilters.Sort([](const TSharedPtr<IFilter<FNavigationToolItemPtr>>& LHS, const TSharedPtr<IFilter<FNavigationToolItemPtr>>& RHS)
+	ChildFilters.Sort([](const TSharedPtr<IFilter<FNavigationToolViewModelPtr>>& LHS, const TSharedPtr<IFilter<FNavigationToolViewModelPtr>>& RHS)
 		{
 			const TSharedPtr<FNavigationToolFilter> CastedLHS = StaticCastSharedPtr<FNavigationToolFilter>(LHS);
 			const TSharedPtr<FNavigationToolFilter> CastedRHS = StaticCastSharedPtr<FNavigationToolFilter>(RHS);
@@ -186,9 +186,9 @@ TArray<TSharedRef<FNavigationToolFilter>> FNavigationToolFilterCollection::GetCa
 void FNavigationToolFilterCollection::ForEachFilter(const TFunctionRef<bool(const TSharedRef<FNavigationToolFilter>&)>& InFunction
 	, const TArray<TSharedRef<FFilterCategory>>& InCategories) const
 {
-	using namespace UE::Sequencer;
+	using namespace Sequencer;
 
-	for (const TSharedPtr<IFilter<FNavigationToolItemPtr>>& Filter : ChildFilters)
+	for (const TSharedPtr<IFilter<FNavigationToolViewModelPtr>>& Filter : ChildFilters)
 	{
 		const TSharedPtr<FNavigationToolFilter> FilterCasted = StaticCastSharedPtr<FNavigationToolFilter>(Filter);
 		if (!FilterCasted.IsValid())
diff --git a/Engine/Plugins/Experimental/SequenceNavigator/Source/SequenceNavigator/Private/Filters/Filters/NavigationToolFilter_CustomText.cpp b/Engine/Plugins/Experimental/SequenceNavigator/Source/SequenceNavigator/Private/Filters/Filters/NavigationToolFilter_CustomText.cpp
--- a/Engine/Plugins/Experimental/SequenceNavigator/Source/SequenceNavigator/Private/Filters/Filters/NavigationToolFilter_CustomText.cpp
+++ b/Engine/Plugins/Experimental/SequenceNavigator/Source/SequenceNavigator/Private/Filters/Filters/NavigationToolFilter_CustomText.cpp
@@ -5,7 +5,7 @@
 namespace UE::SequenceNavigator
 {
 
-using namespace UE::Sequencer;
+using namespace Sequencer;
 
 FNavigationToolFilter_CustomText::FNavigationToolFilter_CustomText(INavigationToolFilterBar& InFilterInterface)
 	: FNavigationToolFilter_Text(InFilterInterface)
@@ -28,7 +28,7 @@ FCustomTextFilterData FNavigationToolFilter_CustomText::CreateCustomTextFilterDa
 	return CustomTextFilterData;
 }
 
-TSharedPtr<FFilterBase<FNavigationToolItemPtr>> FNavigationToolFilter_CustomText::GetFilter()
+TSharedPtr<FFilterBase<FNavigationToolViewModelPtr>> FNavigationToolFilter_CustomText::GetFilter()
 {
 	return AsShared();
 }
@@ -64,7 +64,7 @@ FLinearColor FNavigationToolFilter_CustomText::GetColor() const
 
 FString FNavigationToolFilter_CustomText::GetName() const
 {
-	return FCustomTextFilter<FNavigationToolItemPtr>::GetFilterTypeName().ToString();
+	return FCustomTextFilter<FNavigationToolViewModelPtr>::GetFilterTypeName().ToString();
 }
 
 } // namespace UE::SequenceNavigator
diff --git a/Engine/Plugins/Experimental/SequenceNavigator/Source/SequenceNavigator/Private/Filters/Filters/NavigationToolFilter_Dirty.cpp b/Engine/Plugins/Experimental/SequenceNavigator/Source/SequenceNavigator/Private/Filters/Filters/NavigationToolFilter_Dirty.cpp
--- a/Engine/Plugins/Experimental/SequenceNavigator/Source/SequenceNavigator/Private/Filters/Filters/NavigationToolFilter_Dirty.cpp
+++ b/Engine/Plugins/Experimental/SequenceNavigator/Source/SequenceNavigator/Private/Filters/Filters/NavigationToolFilter_Dirty.cpp
@@ -11,7 +11,7 @@
 namespace UE::SequenceNavigator
 {
 
-using namespace UE::Sequencer;
+using namespace Sequencer;
 
 FNavigationToolFilter_Dirty::FNavigationToolFilter_Dirty(INavigationToolFilterBar& InFilterInterface, TSharedPtr<FFilterCategory> InCategory)
 	: FNavigationToolFilter(InFilterInterface, MoveTemp(InCategory))
@@ -43,21 +43,21 @@ FString FNavigationToolFilter_Dirty::GetName() const
 	return StaticName();
 }
 
-bool FNavigationToolFilter_Dirty::PassesFilter(FNavigationToolFilterType InItem) const
+bool FNavigationToolFilter_Dirty::PassesFilter(const FNavigationToolViewModelPtr InItem) const
 {
-	if (const FNavigationToolSequence* const SequenceItem = InItem->CastTo<FNavigationToolSequence>())
+	if (const TViewModelPtr<FNavigationToolSequence> SequenceItem = InItem.ImplicitCast())
 	{
 		const UMovieSceneSequence* const Sequence = SequenceItem->GetSequence();
 		return Sequence && IsObjectPackageDirty(Sequence);
 	}
 
-	if (const FNavigationToolTrack* const TrackItem = InItem->CastTo<FNavigationToolTrack>())
+	if (const TViewModelPtr<FNavigationToolTrack> TrackItem = InItem.ImplicitCast())
 	{
 		const UMovieSceneTrack* const Track = TrackItem->GetTrack();
 		return Track && IsObjectPackageDirty(Track);
 	}
 
-	if (const FNavigationToolBinding* const BindingItem = InItem->CastTo<FNavigationToolBinding>())
+	if (const TViewModelPtr<FNavigationToolBinding> BindingItem = InItem.ImplicitCast())
 	{
 		const UObject* const BoundObject = BindingItem->GetCachedBoundObject();
 		return BoundObject && IsObjectPackageDirty(BoundObject);
