[SequenceNavigator] Fix Horde loop variable error

diff --git a/Engine/Plugins/Experimental/AvalancheSequenceNavigator/Source/AvalancheSequenceNavigator/Private/AvaNavigationToolProvider.cpp b/Engine/Plugins/Experimental/AvalancheSequenceNavigator/Source/AvalancheSequenceNavigator/Private/AvaNavigationToolProvider.cpp
--- a/Engine/Plugins/Experimental/AvalancheSequenceNavigator/Source/AvalancheSequenceNavigator/Private/AvaNavigationToolProvider.cpp
+++ b/Engine/Plugins/Experimental/AvalancheSequenceNavigator/Source/AvalancheSequenceNavigator/Private/AvaNavigationToolProvider.cpp
@@ -576,7 +576,7 @@ TArray<UAvaSequence*> FAvaNavigationToolProvider::GetSelectedSequences() const
 	TArray<UAvaSequence*> SelectedAvaSequences;
 	SelectedAvaSequences.Reserve(SelectedItems.Num());
 
-	for (const TSharedPtr<FNavigationToolAvaSequence>& SequenceItem : SelectedItems)
+	for (const TViewModelPtr<FNavigationToolAvaSequence>& SequenceItem : SelectedItems)
 	{
 		if (UAvaSequence* const AvaSequence = SequenceItem->GetAvaSequence())
 		{
